Created a generic STweenSlider widget and MVC pattern for displaying it in a toolbar.

diff --git a/Engine/Plugins/Animation/TweeningUtils/Source/TweeningUtilsEditor/Private/Math/Abstraction/BuiltInTweenFunctionContainer.cpp b/Engine/Plugins/Animation/TweeningUtils/Source/TweeningUtilsEditor/Private/Math/Abstraction/BuiltInTweenFunctionContainer.cpp
new file mode 100644
--- /dev/null
+++ b/Engine/Plugins/Animation/TweeningUtils/Source/TweeningUtilsEditor/Private/Math/Abstraction/BuiltInTweenFunctionContainer.cpp
@@ -0,0 +1,73 @@
+﻿// Copyright Epic Games, Inc. All Rights Reserved.
+
+#include "Math/Abstraction/BuiltInTweenFunctionContainer.h"
+
+#include "Framework/Commands/UICommandInfo.h"
+#include "Math/Abstraction/BuiltInTweenFunction.h"
+#include "Misc/Optional.h"
+#include "Templates/Function.h"
+
+namespace UE::TweeningUtilsEditor
+{
+static TOptional<EBlendFunction> GetBlendFunctionType(
+	const TSharedRef<ITweenFunction>& InFunctionToFind, const TArray<TSharedRef<FBuiltInTweenFunction>>& InFunctions
+	)
+{
+	const int32 Index = InFunctions.IndexOfByPredicate([&InFunctionToFind](const TSharedRef<FBuiltInTweenFunction>& Function)
+	{
+		return &InFunctionToFind.Get() == &Function.Get();
+	});
+	// API user can theoretically pass in ANY function so handle that case
+	return ensureMsgf(InFunctions.IsValidIndex(Index), TEXT("You passed a function not managed by this container"))
+		? InFunctions[Index]->GetBlendFunction()
+		: TOptional<EBlendFunction>();
+}
+	
+FBuiltInTweenFunctionContainer::FBuiltInTweenFunctionContainer()
+{
+	ForEachBlendFunction([this](EBlendFunction BlendFunction){ Functions.Add(MakeShared<FBuiltInTweenFunction>(BlendFunction)); });
+}
+
+void FBuiltInTweenFunctionContainer::ForEachFunction(const TFunctionRef<void(const TSharedRef<ITweenFunction>&)> InConsumer) const
+{
+	for (const TSharedRef<FBuiltInTweenFunction>& Function : Functions)
+	{
+		InConsumer(Function);
+	}
+}
+
+TSharedPtr<ITweenFunction> FBuiltInTweenFunctionContainer::GetFunction(int32 InIndex) const
+{
+	return Functions[InIndex];
+}
+
+TSharedPtr<FUICommandInfo> FBuiltInTweenFunctionContainer::GetCommandForFunction(const TSharedRef<ITweenFunction>& InFunction) const
+{
+	const TOptional<EBlendFunction> Function = GetBlendFunctionType(InFunction, Functions);
+	return Function ? TweeningUtilsEditor::GetCommandForBlendFunction(*Function) : nullptr;
+}
+
+const FSlateBrush* FBuiltInTweenFunctionContainer::GetIconForFunction(const TSharedRef<ITweenFunction>& InFunction) const
+{
+	const TOptional<EBlendFunction> Function = GetBlendFunctionType(InFunction, Functions);
+	return Function ? TweeningUtilsEditor::GetUntintedIconForTweenFunction(*Function) : nullptr;
+}
+
+FLinearColor FBuiltInTweenFunctionContainer::GetColorForFunction(const TSharedRef<ITweenFunction>& InFunction) const
+{
+	const TOptional<EBlendFunction> Function = GetBlendFunctionType(InFunction, Functions);
+	return Function ? TweeningUtilsEditor::GetTintColorForTweenFunction(*Function) : FLinearColor::White;
+}
+
+FText FBuiltInTweenFunctionContainer::GetLabelForFunction(const TSharedRef<ITweenFunction>& InFunction) const
+{
+	const TOptional<EBlendFunction> Function = GetBlendFunctionType(InFunction, Functions);
+	return Function ? TweeningUtilsEditor::GetLabelForBlendFunction(*Function) : FText::GetEmpty();
+}
+	
+FText FBuiltInTweenFunctionContainer::GetToolTipForFunction(const TSharedRef<ITweenFunction>& InFunction) const
+{
+	const TOptional<EBlendFunction> Function = GetBlendFunctionType(InFunction, Functions);
+	return Function ? TweeningUtilsEditor::GetDescriptionForBlendFunction(*Function) : FText::GetEmpty();
+}
+}
diff --git a/Engine/Plugins/Animation/TweeningUtils/Source/TweeningUtilsEditor/Private/Math/Abstraction/KeyBlendingAbstraction.cpp b/Engine/Plugins/Animation/TweeningUtils/Source/TweeningUtilsEditor/Private/Math/Abstraction/KeyBlendingAbstraction.cpp
new file mode 100644
--- /dev/null
+++ b/Engine/Plugins/Animation/TweeningUtils/Source/TweeningUtilsEditor/Private/Math/Abstraction/KeyBlendingAbstraction.cpp
@@ -0,0 +1,150 @@
+﻿// Copyright Epic Games, Inc. All Rights Reserved.
+
+#include "Math/Abstraction/KeyBlendingAbstraction.h"
+
+#include "Containers/UnrealString.h"
+#include "Internationalization/Internationalization.h"
+#include "Internationalization/Text.h"
+#include "Math/Color.h"
+#include "Math/KeyBlendInfo.h"
+#include "Math/KeyBlendingFunctions.h"
+#include "Styling/SlateBrush.h"
+#include "Templates/SharedPointer.h"
+#include "TweeningUtilsCommands.h"
+#include "TweeningUtilsStyle.h"
+
+#define LOCTEXT_NAMESPACE "FStaticBlendFunctionData"
+
+namespace UE::TweeningUtilsEditor
+{
+struct FStaticBlendFunctionData
+{
+	using FBlendFunctionPtr = double(*)(const FKeyBlendInfo& InBlendInfo);
+	
+	const EBlendFunction Type;
+	const FBlendFunctionPtr EvaluationFunctionPtr;
+	const TWeakPtr<FUICommandInfo> Command; // Don't keep the command in a strong pointer case FTweeningUtilsCommands singleton gets destroyed
+	const FText Label;
+	const FText Abbreviation;
+	const FText Description;
+	const FString StyleBaseName;
+	
+	FStaticBlendFunctionData(
+		EBlendFunction InType,
+		FBlendFunctionPtr InEvaluationFunctionPtr,
+		TSharedPtr<FUICommandInfo> InCommand,
+		FText InLabel,
+		FText InAbbreviation,
+		FText InDescription
+		)
+		: Type(InType)
+		, EvaluationFunctionPtr(InEvaluationFunctionPtr)
+		, Command(MoveTemp(InCommand))
+		, Label(MoveTemp(InLabel))
+		, Abbreviation(MoveTemp(InAbbreviation))
+		, Description(MoveTemp(InDescription))
+		, StyleBaseName(FString::Printf(TEXT("TweeningUtils.%s"), InCommand ? *InCommand->GetCommandName().ToString() : TEXT("")))
+	{
+		check(InCommand.IsValid());
+	}
+	
+	const FSlateBrush* GetUntintedIconForTweenFunction() const
+	{
+		const FString StyleString = StyleBaseName + TEXT(".Icon");
+		return FTweeningUtilsStyle::Get().GetBrush(*StyleString);
+	}
+
+	FLinearColor GetTintColorForTweenFunction() const
+	{
+		const FString StyleString = StyleBaseName + TEXT(".Color");
+		return FTweeningUtilsStyle::Get().GetColor(*StyleString);
+	}
+};
+
+static FStaticBlendFunctionData& GetFunctionData(EBlendFunction InFunction)
+{
+	static_assert(static_cast<int32>(EBlendFunction::Num) == 6, "Extend this array");
+	constexpr int32 Num = static_cast<int32>(EBlendFunction::Num);
+	static FStaticBlendFunctionData Data[Num] =
+	{
+		FStaticBlendFunctionData(
+			EBlendFunction::ControlsToTween, &Blend_ControlsToTween, FTweeningUtilsCommands::Get().SetControlsToTween,
+			LOCTEXT("ControlsToTween.Label", "Tween"),
+			LOCTEXT("ControlsToTween.Abbreviation", "TW"),
+			LOCTEXT("ControlsToTween.Description", "Interpolates between the previous and next keys")
+		),
+		FStaticBlendFunctionData(
+			EBlendFunction::PushPull, &Blend_PushPull, FTweeningUtilsCommands::Get().SetTweenPushPull,
+			LOCTEXT("PushPull.Label", "Push / Pull"),
+			LOCTEXT("PushPull.Abbreviation", "PP"),
+			LOCTEXT("PushPull.Description", "Push or pull the values to the interpolation between the previous and next keys")
+		),
+		FStaticBlendFunctionData(
+			EBlendFunction::BlendNeighbor, &Blend_Neighbor, FTweeningUtilsCommands::Get().SetTweenBlendNeighbor,
+			LOCTEXT("BlendNeighbor.Label", "Blend Neighbor"),
+			LOCTEXT("BlendNeighbour.Abbreviation", "BN"),
+			LOCTEXT("BlendNeighbour.Description", "Blend to the next or previous values for selected keys.")
+		),
+		FStaticBlendFunctionData(
+			EBlendFunction::BlendRelative, &Blend_Relative, FTweeningUtilsCommands::Get().SetTweenBlendRelative,
+			LOCTEXT("BlendRelative.Label", "Blend Relative"),
+			LOCTEXT("BlendRelative.Abbreviation", "BR"),
+			LOCTEXT("BlendRelative.Description", "Blend relative to the next or previous value for selected keys")
+		),
+		FStaticBlendFunctionData(
+			EBlendFunction::BlendEase, &Blend_Ease, FTweeningUtilsCommands::Get().SetTweenBlendEase,
+			LOCTEXT("BlendEase.Label", "Blend Ease"),
+			LOCTEXT("BlendEase.Abbreviation", "BE"),
+			LOCTEXT("BlendEase.Description", "Blend with an ease falloff to the next or previous value for selected keys")
+		),
+		FStaticBlendFunctionData(
+			EBlendFunction::SmoothRough, &Blend_SmoothRough, FTweeningUtilsCommands::Get().SetTweenSmoothRough,
+			LOCTEXT("SmoothRough.Label", "Smooth Rough"),
+			LOCTEXT("SmoothRough.Abbreviation", "SR"),
+			LOCTEXT("SmoothRough.Description", "Make selected keys smooth or rough")
+		)
+	};
+
+	const int32 Index = static_cast<int32>(InFunction);
+	return Data[Index];
+}
+	
+double EvaluateBlend(EBlendFunction InFunction, const FKeyBlendInfo& InBlendInfo)
+{
+	return GetFunctionData(InFunction).EvaluationFunctionPtr(InBlendInfo);
+}
+
+TSharedPtr<FUICommandInfo> GetCommandForBlendFunction(EBlendFunction InFunction)
+{
+	const TSharedPtr<FUICommandInfo> Command = GetFunctionData(InFunction).Command.Pin();
+	ensure(Command);
+	return Command;
+}
+
+const FSlateBrush* GetUntintedIconForTweenFunction(EBlendFunction InFunction)
+{
+	return GetFunctionData(InFunction).GetUntintedIconForTweenFunction();
+}
+
+FLinearColor GetTintColorForTweenFunction(EBlendFunction InFunction)
+{
+	return GetFunctionData(InFunction).GetTintColorForTweenFunction();
+}
+	
+FText GetLabelForBlendFunction(EBlendFunction InFunction)
+{
+	return GetFunctionData(InFunction).Label;
+}
+
+FText GetAbbreviationForBlendFunction(EBlendFunction InFunction)
+{
+	return GetFunctionData(InFunction).Abbreviation;
+}
+
+FText GetDescriptionForBlendFunction(EBlendFunction InFunction)
+{
+	return GetFunctionData(InFunction).Description;
+}
+}
+
+#undef LOCTEXT_NAMESPACE
\ No newline at end of file
diff --git a/Engine/Plugins/Animation/TweeningUtils/Source/TweeningUtilsEditor/Private/Math/KeyBlendingAbstraction.cpp b/Engine/Plugins/Animation/TweeningUtils/Source/TweeningUtilsEditor/Private/Math/KeyBlendingAbstraction.cpp
deleted file mode 100644
diff --git a/Engine/Plugins/Animation/TweeningUtils/Source/TweeningUtilsEditor/Private/TweeningUtilsCommands.cpp b/Engine/Plugins/Animation/TweeningUtils/Source/TweeningUtilsEditor/Private/TweeningUtilsCommands.cpp
new file mode 100644
--- /dev/null
+++ b/Engine/Plugins/Animation/TweeningUtils/Source/TweeningUtilsEditor/Private/TweeningUtilsCommands.cpp
@@ -0,0 +1,36 @@
+﻿// Copyright Epic Games, Inc. All Rights Reserved.
+
+#include "TweeningUtilsCommands.h"
+
+#include "TweeningUtilsStyle.h"
+#include "Internationalization/Internationalization.h"
+#include "Math/Abstraction/KeyBlendingAbstraction.h"
+
+#define LOCTEXT_NAMESPACE "FTweeningUtilsCommands"
+
+namespace UE::TweeningUtilsEditor
+{
+FTweeningUtilsCommands::FTweeningUtilsCommands()
+	: TCommands<FTweeningUtilsCommands>
+	(
+		TEXT("TweeningUtils"),
+		NSLOCTEXT("Contexts", "TweeningUtils", "Tweening"),
+		NAME_None,
+		FTweeningUtilsStyle::Get().GetStyleSetName()
+	)
+{}
+
+void FTweeningUtilsCommands::RegisterCommands()
+{
+	static_assert(static_cast<int32>(EBlendFunction::Num) == 6, "Add a command here");
+	UI_COMMAND(SetControlsToTween, "Set Tween ", "Sets the tween mode to Tween", EUserInterfaceActionType::RadioButton, FInputChord() );
+	UI_COMMAND(SetTweenPushPull, "Set Push / Pull", "Sets the tween mode to Push / Pull", EUserInterfaceActionType::RadioButton, FInputChord() );
+	UI_COMMAND(SetTweenBlendNeighbor, "Set Blend Neighbor", "Sets the tween mode to Blend Neighbour", EUserInterfaceActionType::RadioButton, FInputChord() );
+	UI_COMMAND(SetTweenBlendRelative, "Set Blend Relative", "Sets the tween mode to Blend Relative", EUserInterfaceActionType::RadioButton, FInputChord() );
+	UI_COMMAND(SetTweenBlendEase, "Set Blend Ease", "Sets the tween mode to Blend Ease", EUserInterfaceActionType::RadioButton, FInputChord() );
+	UI_COMMAND(SetTweenSmoothRough, "Set Smooth / Rough", "Sets the tween mode to Smooth / Rough", EUserInterfaceActionType::RadioButton, FInputChord() );
+	UI_COMMAND(ToggleOvershootMode, "Overshoot mode", "Overshoot extends the available range over 100%", EUserInterfaceActionType::ToggleButton, FInputChord() ); 
+}
+}
+
+#undef LOCTEXT_NAMESPACE
diff --git a/Engine/Plugins/Animation/TweeningUtils/Source/TweeningUtilsEditor/Private/TweeningUtilsEditorModule.cpp b/Engine/Plugins/Animation/TweeningUtils/Source/TweeningUtilsEditor/Private/TweeningUtilsEditorModule.cpp
--- a/Engine/Plugins/Animation/TweeningUtils/Source/TweeningUtilsEditor/Private/TweeningUtilsEditorModule.cpp
+++ b/Engine/Plugins/Animation/TweeningUtils/Source/TweeningUtilsEditor/Private/TweeningUtilsEditorModule.cpp
@@ -2,16 +2,20 @@
 
 #include "TweeningUtilsEditorModule.h"
 
+#include "TweeningUtilsCommands.h"
+#include "TweeningUtilsStyle.h"
+
 namespace UE::TweeningUtilsEditor
 {
 void FTweeningUtilsEditorModule::StartupModule()
 {
-	
+	FTweeningUtilsCommands::Register();
+	FTweeningUtilsStyle::Get();
 }
 
 void FTweeningUtilsEditorModule::ShutdownModule()
 {
-	
+	FTweeningUtilsCommands::Unregister();
 }
 	
 IMPLEMENT_MODULE(FTweeningUtilsEditorModule, TweeningUtilsEditor)
diff --git a/Engine/Plugins/Animation/TweeningUtils/Source/TweeningUtilsEditor/Private/TweeningUtilsStyle.cpp b/Engine/Plugins/Animation/TweeningUtils/Source/TweeningUtilsEditor/Private/TweeningUtilsStyle.cpp
new file mode 100644
--- /dev/null
+++ b/Engine/Plugins/Animation/TweeningUtils/Source/TweeningUtilsEditor/Private/TweeningUtilsStyle.cpp
@@ -0,0 +1,72 @@
+﻿// Copyright Epic Games, Inc. All Rights Reserved.
+
+#include "TweeningUtilsStyle.h"
+
+#include "Math/Abstraction/KeyBlendingAbstraction.h"
+#include "Styling/SlateStyleMacros.h"
+#include "Styling/SlateStyleRegistry.h"
+#include "Widgets/TweenSliderStyle.h"
+
+namespace UE::TweeningUtilsEditor
+{
+FTweeningUtilsStyle& FTweeningUtilsStyle::Get()
+{
+	static FTweeningUtilsStyle Style;
+	return Style;
+}
+
+FTweeningUtilsStyle::FTweeningUtilsStyle()
+	: FSlateStyleSet("TweeningUtilsStyle")
+{
+	const FString PluginContentDir = FPaths::EnginePluginsDir() / TEXT("Animation") / TEXT("TweeningUtils") / TEXT("Resources");
+	const FString EngineEditorSlateDir = FPaths::EngineContentDir() / TEXT("Editor") / TEXT("Slate");
+	SetContentRoot(PluginContentDir);
+	SetCoreContentRoot(EngineEditorSlateDir);
+	
+	const FVector2D Icon20x20(20.0f, 20.0f);
+	
+	// Tween - Slider colors
+	static_assert(static_cast<int32>(EBlendFunction::Num) == 6, "Add a color here");
+	const FLinearColor ControlsToTweenTint(254.f/255.f, 155.f/255.f, 7.f/255.f);
+	const FLinearColor PushPullTint(107.f/255.f, 223.f/255.f, 151.f/255.f);
+	const FLinearColor BlendNeighborTint(133.f/255.f, 185.f/255.f, 73.f/255.f);
+	const FLinearColor BlendRelativeTint(39.f/255.f, 177.f/255.f, 240.f/255.f);
+	const FLinearColor BlendEaseTint(148.f/255.f, 57.f/255.f, 175.f/255.f);
+	const FLinearColor SmoothRoughTint(235.f/255.f, 64.f/255.f, 64.f/255.f);
+	Set("TweeningUtils.SetControlsToTween.Color", ControlsToTweenTint);
+	Set("TweeningUtils.SetTweenPushPull.Color", PushPullTint);
+	Set("TweeningUtils.SetTweenBlendNeighbor.Color", BlendNeighborTint);
+	Set("TweeningUtils.SetTweenBlendRelative.Color", BlendRelativeTint);
+	Set("TweeningUtils.SetTweenBlendEase.Color", BlendEaseTint);
+	Set("TweeningUtils.SetTweenSmoothRough.Color", SmoothRoughTint);
+	
+	// Tween - Commands
+	static_assert(static_cast<int32>(EBlendFunction::Num) == 6, "Add command style");
+	Set("TweeningUtils.SetControlsToTween", new IMAGE_BRUSH_SVG("Icons/ControlsToTween", Icon20x20, ControlsToTweenTint));
+	Set("TweeningUtils.SetTweenPushPull", new IMAGE_BRUSH_SVG("Icons/PushPull", Icon20x20, PushPullTint));
+	Set("TweeningUtils.SetTweenBlendNeighbor", new IMAGE_BRUSH_SVG("Icons/BlendNeighbor", Icon20x20, BlendNeighborTint));
+	Set("TweeningUtils.SetTweenBlendRelative", new IMAGE_BRUSH_SVG("Icons/BlendRelative", Icon20x20, BlendRelativeTint));
+	Set("TweeningUtils.SetTweenBlendEase", new IMAGE_BRUSH_SVG("Icons/BlendEase", Icon20x20, BlendEaseTint));
+	Set("TweeningUtils.SetTweenSmoothRough", new IMAGE_BRUSH_SVG("Icons/SmoothRough", Icon20x20, SmoothRoughTint));
+	Set("TweeningUtils.ToggleOvershootMode", new IMAGE_BRUSH_SVG("Icons/OvershootMode", Icon20x20));
+	
+	// Tween - Icons (no tint)
+	static_assert(static_cast<int32>(EBlendFunction::Num) == 6, "Add untinted icon here");
+	Set("TweeningUtils.SetControlsToTween.Icon", new IMAGE_BRUSH_SVG("Icons/ControlsToTween", Icon20x20));
+	Set("TweeningUtils.SetTweenPushPull.Icon", new IMAGE_BRUSH_SVG("Icons/PushPull", Icon20x20));
+	Set("TweeningUtils.SetTweenBlendNeighbor.Icon", new IMAGE_BRUSH_SVG("Icons/BlendNeighbor", Icon20x20));
+	Set("TweeningUtils.SetTweenBlendRelative.Icon", new IMAGE_BRUSH_SVG("Icons/BlendRelative", Icon20x20));
+	Set("TweeningUtils.SetTweenBlendEase.Icon", new IMAGE_BRUSH_SVG("Icons/BlendEase", Icon20x20));
+	Set("TweeningUtils.SetTweenSmoothRough.Icon", new IMAGE_BRUSH_SVG("Icons/SmoothRough", Icon20x20));
+	
+	// Tween - Slider visuals
+	Set("TweenSlider", FTweenSliderStyle());
+	
+	FSlateStyleRegistry::RegisterSlateStyle(*this);
+}
+
+FTweeningUtilsStyle::~FTweeningUtilsStyle()
+{
+	FSlateStyleRegistry::UnRegisterSlateStyle(*this);
+}
+}
diff --git a/Engine/Plugins/Animation/TweeningUtils/Source/TweeningUtilsEditor/Private/TweeningUtilsStyle.h b/Engine/Plugins/Animation/TweeningUtils/Source/TweeningUtilsEditor/Private/TweeningUtilsStyle.h
new file mode 100644
--- /dev/null
+++ b/Engine/Plugins/Animation/TweeningUtils/Source/TweeningUtilsEditor/Private/TweeningUtilsStyle.h
@@ -0,0 +1,20 @@
+﻿// Copyright Epic Games, Inc. All Rights Reserved.
+
+#pragma once
+
+#include "Styling/SlateStyle.h"
+
+namespace UE::TweeningUtilsEditor
+{
+/** Style for the widgets and commands in this module. */
+class FTweeningUtilsStyle final
+: public FSlateStyleSet
+{
+public:
+	
+	static FTweeningUtilsStyle& Get();
+	
+	FTweeningUtilsStyle();
+	virtual ~FTweeningUtilsStyle() override;
+};
+}
diff --git a/Engine/Plugins/Animation/TweeningUtils/Source/TweeningUtilsEditor/Private/Widgets/MVC/Models/CurveTweeningModel.cpp b/Engine/Plugins/Animation/TweeningUtils/Source/TweeningUtilsEditor/Private/Widgets/MVC/Models/CurveTweeningModel.cpp
new file mode 100644
--- /dev/null
+++ b/Engine/Plugins/Animation/TweeningUtils/Source/TweeningUtilsEditor/Private/Widgets/MVC/Models/CurveTweeningModel.cpp
@@ -0,0 +1,42 @@
+﻿// Copyright Epic Games, Inc. All Rights Reserved.
+
+#include "Widgets/MVC/Models/CurveTweenModel.h"
+
+#include "Math/CurveBlending.h"
+#include "Math/Abstraction/ITweenFunction.h"
+
+namespace UE::TweeningUtilsEditor
+{
+FCurveTweenModel::FCurveTweenModel(TWeakPtr<FCurveEditor> InWeakCurveEditor, const TSharedRef<ITweenFunction>& InTweenFunction)
+	: FFunctionBasedTweenModel(InTweenFunction)
+	, WeakCurveEditor(InWeakCurveEditor)
+{}
+
+void FCurveTweenModel::StartBlendOperation()
+{
+	const TSharedPtr<FCurveEditor> CurveEditor = WeakCurveEditor.Pin();
+	ContiguousKeySelection = CurveEditor ? FContiguousKeyMapping(*CurveEditor) : FContiguousKeyMapping{};
+}
+
+void FCurveTweenModel::StopBlendOperation()
+{
+	ContiguousKeySelection.KeyMap.Empty(
+		ContiguousKeySelection.KeyMap.Num()
+		);
+}
+
+void FCurveTweenModel::BlendValues(float InNormalizedValue)
+{
+	const TSharedPtr<FCurveEditor> CurveEditor = WeakCurveEditor.Pin();
+	if (!CurveEditor)
+	{
+		return;
+	}
+
+	const float ScaledBlendValue = ScaleBlendValue(InNormalizedValue); 
+	TweeningUtilsEditor::BlendCurves(ScaledBlendValue, *CurveEditor, ContiguousKeySelection, [this](const FKeyBlendInfo& InBlendInfo)
+	{
+		return GetTweenFunction()->Blend(InBlendInfo);
+	});
+}
+}
diff --git a/Engine/Plugins/Animation/TweeningUtils/Source/TweeningUtilsEditor/Private/Widgets/MVC/STweenView.cpp b/Engine/Plugins/Animation/TweeningUtils/Source/TweeningUtilsEditor/Private/Widgets/MVC/STweenView.cpp
new file mode 100644
--- /dev/null
+++ b/Engine/Plugins/Animation/TweeningUtils/Source/TweeningUtilsEditor/Private/Widgets/MVC/STweenView.cpp
@@ -0,0 +1,72 @@
+﻿// Copyright Epic Games, Inc. All Rights Reserved.
+
+#include "Widgets/MVC/STweenView.h"
+
+#include "Styling/AppStyle.h"
+#include "Widgets/Layout/SBox.h"
+#include "Widgets/MVC/Models/TweenModel.h"
+#include "Widgets/STweenSlider.h"
+
+namespace UE::TweeningUtilsEditor
+{
+void STweenView::Construct(const FArguments& InArgs, const TSharedRef<FTweenModel>& InTweenModel)
+{
+	check(InArgs._SliderIcon.IsSet() || InArgs._SliderIcon.IsBound());
+	check(InArgs._SliderColor.IsSet() || InArgs._SliderColor.IsBound());
+	
+	TweenModel = InTweenModel;
+
+	// STweenView is placed in toolbars.
+	// We want STweenSlider's foreground color to be the same as that of the SComboButton from which you pick functions for color consistency.
+	// Effectively, FComboButtons in the toolbar use this style color.
+	SetForegroundColor(FAppStyle::GetSlateColor("DefaultForeground"));
+	
+	ChildSlot
+	[
+		SNew(SBox)
+		.HAlign(HAlign_Fill)
+		.VAlign(VAlign_Center) 
+		[
+			SNew(STweenSlider)
+			.SliderIcon(InArgs._SliderIcon)
+			.SliderColor(InArgs._SliderColor)
+			.ScaleRenderMode(this, &STweenView::GetBarRenderMode)
+			.OnSliderDragStarted(this, &STweenView::OnDragStarted)
+			.OnSliderDragEnded(this, &STweenView::OnDragEnded)
+			.OnSliderValueDragged(this, &STweenView::OnDragValueUpdated)
+			.OnPointValuePicked(this, &STweenView::OnPointPicked)
+			.MapSliderValueToBlendValue(this, &STweenView::MapSliderValueToBlendValue)
+		]
+	];
+}
+
+ETweenScaleMode STweenView::GetBarRenderMode() const
+{
+	return TweenModel->GetScaleMode();
+}
+
+void STweenView::OnDragStarted()
+{
+	TweenModel->StartBlendOperation();
+}
+
+void STweenView::OnDragEnded()
+{
+	TweenModel->StopBlendOperation();
+}
+
+void STweenView::OnDragValueUpdated(float Value)
+{
+	TweenModel->BlendValues(Value);
+}
+
+void STweenView::OnPointPicked(float Value)
+{
+	TweenModel->BlendOneOff(Value);
+}
+
+float STweenView::MapSliderValueToBlendValue(float Value)
+{
+	return TweenModel->ScaleBlendValue(Value);
+}
+}
diff --git a/Engine/Plugins/Animation/TweeningUtils/Source/TweeningUtilsEditor/Private/Widgets/MVC/TweenToolbarController.cpp b/Engine/Plugins/Animation/TweeningUtils/Source/TweeningUtilsEditor/Private/Widgets/MVC/TweenToolbarController.cpp
new file mode 100644
--- /dev/null
+++ b/Engine/Plugins/Animation/TweeningUtils/Source/TweeningUtilsEditor/Private/Widgets/MVC/TweenToolbarController.cpp
@@ -0,0 +1,135 @@
+﻿// Copyright Epic Games, Inc. All Rights Reserved.
+
+#include "Widgets/MVC/TweenToolbarController.h"
+
+#include "Framework/Commands/UICommandList.h"
+#include "Framework/MultiBox/MultiBoxBuilder.h"
+#include "Math/Abstraction/ITweenFunctionContainer.h"
+#include "Math/Color.h"
+#include "TweeningUtilsCommands.h"
+#include "Widgets/MVC/Models/FunctionBasedTweenModel.h"
+#include "Widgets/MVC/STweenView.h"
+
+namespace UE::TweeningUtilsEditor
+{
+FTweenToolbarController::FTweenToolbarController(
+	const TSharedRef<FFunctionBasedTweenModel>& InModel,
+	const TSharedRef<FUICommandList>& InCommandList,
+	const TSharedRef<ITweenFunctionContainer>& InTweenFunctions
+	)
+	: TweenModel(InModel)
+	, CommandList(InCommandList)
+	, TweenFunctions(InTweenFunctions)
+{
+	// Function changing commands
+	TweenFunctions->ForEachFunction([this](const TSharedRef<ITweenFunction>& InFunction)
+	{
+		CommandList->MapAction(TweenFunctions->GetCommandForFunction(
+			InFunction),
+			FUIAction(
+			   FExecuteAction::CreateRaw(this, &FTweenToolbarController::SetBlendFunction, InFunction.ToSharedPtr()),
+			   {},
+			   FIsActionChecked::CreateRaw(this, &FTweenToolbarController::IsBlendFunctionSelected, InFunction.ToSharedPtr())
+			)
+		);
+	});
+
+	InCommandList->MapAction(FTweeningUtilsCommands::Get().ToggleOvershootMode,
+		FUIAction(
+			FExecuteAction::CreateRaw(this, &FTweenToolbarController::ToggleOvershootMode),
+			FCanExecuteAction(),
+			FIsActionChecked::CreateRaw(this, &FTweenToolbarController::IsOvershootModeEnabled)
+			)
+);
+}
+
+FTweenToolbarController::~FTweenToolbarController()
+{
+	// Function changing commands
+	TweenFunctions->ForEachFunction([this](const TSharedRef<ITweenFunction>& InFunction)
+	{
+		CommandList->UnmapAction(TweenFunctions->GetCommandForFunction(InFunction));
+	});
+	
+	CommandList->UnmapAction(FTweeningUtilsCommands::Get().ToggleOvershootMode);
+}
+
+void FTweenToolbarController::AddToToolbar(FToolBarBuilder& ToolBarBuilder) const
+{
+	ToolBarBuilder.SetLabelVisibility(EVisibility::Visible);
+	ToolBarBuilder.AddComboButton(
+	   FUIAction(),
+	   FOnGetContent::CreateRaw(this, &FTweenToolbarController::MakeTweenModeMenu),
+	   TAttribute<FText>::CreateRaw(this, &FTweenToolbarController::GetLabelForComboBox),
+	   TAttribute<FText>::CreateRaw(this, &FTweenToolbarController::GetToolTipForComboBox),
+	   TAttribute<FSlateIcon>::CreateRaw(this, &FTweenToolbarController::GetIconForComboBox)
+	   );
+	ToolBarBuilder.SetLabelVisibility(EVisibility::Collapsed);
+		
+	ToolBarBuilder.AddWidget(
+		SNew(STweenView, TweenModel)
+		.SliderIcon(TAttribute<const FSlateBrush*>::CreateRaw(this, &FTweenToolbarController::GetIconForSlider))
+		.SliderColor(TAttribute<FLinearColor>::CreateRaw(this, &FTweenToolbarController::GetColorForSlider))
+		);
+	ToolBarBuilder.AddToolBarButton(FTweeningUtilsCommands::Get().ToggleOvershootMode);
+}
+
+TSharedRef<SWidget> FTweenToolbarController::MakeTweenModeMenu() const
+{
+	FMenuBuilder MenuBuilder(true, CommandList);
+
+	TweenFunctions->ForEachFunction([this, &MenuBuilder](const TSharedRef<ITweenFunction>& InFunction)
+	{
+		MenuBuilder.AddMenuEntry(TweenFunctions->GetCommandForFunction(InFunction));
+	});
+
+	return MenuBuilder.MakeWidget();
+}
+
+FText FTweenToolbarController::GetLabelForComboBox() const
+{
+	return TweenFunctions->GetLabelForFunction(TweenModel->GetTweenFunction());
+}
+
+FText FTweenToolbarController::GetToolTipForComboBox() const
+{
+	return TweenFunctions->GetToolTipForFunction(TweenModel->GetTweenFunction());
+}
+
+FSlateIcon FTweenToolbarController::GetIconForComboBox() const
+{
+	return TweenFunctions->GetCommandForFunction(TweenModel->GetTweenFunction())->GetIcon();
+}
+
+void FTweenToolbarController::SetBlendFunction(TSharedPtr<ITweenFunction> InFunction) const
+{
+	TweenModel->SetTweenFunction(InFunction.ToSharedRef());
+}
+
+bool FTweenToolbarController::IsBlendFunctionSelected(TSharedPtr<ITweenFunction> InFunction) const
+{
+	return TweenModel->GetTweenFunction() == InFunction;
+}
+
+void FTweenToolbarController::ToggleOvershootMode() const
+{
+	TweenModel->SetScaleMode(
+	   TweenModel->GetScaleMode() == ETweenScaleMode::Normalized ? ETweenScaleMode::Overshoot : ETweenScaleMode::Normalized
+	   );
+}
+
+bool FTweenToolbarController::IsOvershootModeEnabled() const
+{
+	return TweenModel->GetScaleMode() == ETweenScaleMode::Overshoot;
+}
+
+const FSlateBrush* FTweenToolbarController::GetIconForSlider() const
+{
+	return TweenFunctions->GetIconForFunction(TweenModel->GetTweenFunction());
+}
+
+FLinearColor FTweenToolbarController::GetColorForSlider() const
+{
+	return TweenFunctions->GetColorForFunction(TweenModel->GetTweenFunction());
+}
+}
diff --git a/Engine/Plugins/Animation/TweeningUtils/Source/TweeningUtilsEditor/Private/Widgets/STweenSlider.cpp b/Engine/Plugins/Animation/TweeningUtils/Source/TweeningUtilsEditor/Private/Widgets/STweenSlider.cpp
new file mode 100644
--- /dev/null
+++ b/Engine/Plugins/Animation/TweeningUtils/Source/TweeningUtilsEditor/Private/Widgets/STweenSlider.cpp
@@ -0,0 +1,382 @@
+﻿// Copyright Epic Games, Inc. All Rights Reserved.
+
+#include "Widgets/STweenSlider.h"
+
+#include "Internationalization/Internationalization.h"
+#include "Internationalization/Text.h"
+#include "Widgets/SToolTip.h"
+#include "Widgets/TweenSliderDrawUtils.h"
+
+#define LOCTEXT_NAMESPACE "STweenSlider"
+
+namespace UE::TweeningUtilsEditor
+{
+struct FSliderWidgetData
+{
+	/** Contains all values required for drawing the slider. */
+	FTweenWidgetArgs WidgetArgs;
+
+	/** Hover state as it was when the mouse was last moved. */
+	FTweenSliderHoverState HoverState;
+	/** Same length as the number of points in the slider. Corresponds to slider value, which is in [0,1] range. Updated when hover state is updated. */
+	TArray<float> PointSliderValues;
+
+	// Technically, we only need the pimpl to hide the above stuff but since we already have it,
+	// we'll just put all impl data here so we're more flexible, avoid recompiles, make hotfixes easier, etc.
+	
+	struct FDelayedDragData
+	{
+		/** Used to detect drags. */
+		FDelayedDrag DelayedDrag;
+
+		/** Used to determine whether the slider value changes this tick. */
+		float LastSliderPosition = 0.5f;
+		
+		explicit FDelayedDragData(const FVector2D& InInitialPosition, const FKey& InEffectiveKey)
+			: DelayedDrag(InInitialPosition, InEffectiveKey)
+		{}
+	};
+	/** Data valid while the mouse button is down. */
+	TOptional<FDelayedDragData> DragData;
+	/** Whether the mouse was on a point when pressed. */
+	bool bStartedMouseDownOnPoint = false;
+
+	/** The last mouse position the user had. */
+	FVector2D LastMousePosition;
+	/** The last slider position the user hovered the mouse with (regardless of whether dragging or not). */
+	float LastMousePositionOnSliderBar = 0.5f;
+	
+	/**
+	 * The position of the slider in range 0.0 (completely left) to 1.0 (completely right).
+	 * This value is converted to -1 to 1 range for OnSliderChangedDelegate.
+	 */
+	float TargetSliderPosition = 0.5f;
+	/** Interpolated value where the slider is actually drawn. This makes the slider feel smoother. */
+	float AnimatedCurrentSliderPosition = 0.5;
+
+	/** A custom tooltip is used so it can be marked as interactive while sliding the button (otherwise it would disappear). */
+	TSharedPtr<SToolTip> ToolTip;
+};
+
+/** Converts [0,1] range to [-1,1] range. */
+static float SliderValueToBlendValue(float InSliderValue)
+{
+	return InSliderValue * 2.0 - 1.0;
+}
+	
+void STweenSlider::Construct(const FArguments& InArgs)
+{
+	check(InArgs._Style);
+	check(InArgs._ColorAndOpacity.IsSet() || InArgs._ColorAndOpacity.IsBound());
+	checkf(InArgs._SliderIcon.IsSet() || InArgs._SliderIcon.IsBound(), TEXT("You must provide a getter - ok to return nullptr though"));
+	check(InArgs._SliderColor.IsSet() || InArgs._SliderColor.IsBound());
+	check(InArgs._ScaleRenderMode.IsSet() || InArgs._ScaleRenderMode.IsBound());
+	
+	Pimpl = MakePimpl<FSliderWidgetData>();
+	
+	FTweenWidgetArgs& WidgetArgs = Pimpl->WidgetArgs;
+	WidgetArgs.Style = InArgs._Style;
+	WidgetArgs.ColorAndOpacity = InArgs._ColorAndOpacity;
+	WidgetArgs.SliderIconAttr = InArgs._SliderIcon;
+	WidgetArgs.SliderColor = InArgs._SliderColor;
+	WidgetArgs.ScaleRenderModeAttr = InArgs._ScaleRenderMode;
+
+	Pimpl->ToolTip = SNew(SToolTip)
+		.Text(this, &STweenSlider::GetToolTipText)
+		// This prevents the tooltip from disappearing when the mouse is pressed down.
+		.IsInteractive(this, &STweenSlider::IsMouseDown);
+
+	OnSliderDragStartedDelegate = InArgs._OnSliderDragStarted;
+	OnSliderDragStoppedDelegate = InArgs._OnSliderDragEnded;
+	OnSliderChangedDelegate = InArgs._OnSliderValueDragged;
+	OnPointValuePickedDelegate = InArgs._OnPointValuePicked;
+	MapSliderValueToBlendValueDelegate = InArgs._MapSliderValueToBlendValue;
+	if (!MapSliderValueToBlendValueDelegate.IsBound())
+	{
+		MapSliderValueToBlendValueDelegate = FMapSliderValueToBlendValue::CreateLambda([](double Value){ return Value; });
+	}
+}
+
+FReply STweenSlider::OnMouseButtonDoubleClick(const FGeometry& InMyGeometry, const FPointerEvent& InMouseEvent)
+{
+	// User might click the bar background in quick succession to blend e.g. value 0.835 multiple times;
+	// the clicks may be so quick that it triggers a double click.
+	return OnMouseButtonDown(InMyGeometry, InMouseEvent);
+}
+
+FReply STweenSlider::OnMouseButtonDown(const FGeometry& InGeometry, const FPointerEvent& InMouseEvent)
+{
+	if (InMouseEvent.GetEffectingButton() != EKeys::LeftMouseButton)
+	{
+		return FReply::Unhandled();
+	}
+	
+	FTweenWidgetArgs& WidgetArgs = Pimpl->WidgetArgs;
+	FGeometry SliderArea, IconArea;
+	TArray<FGeometry> Points;
+	GetSliderButtonGeometry(Pimpl->AnimatedCurrentSliderPosition, InGeometry, WidgetArgs, SliderArea, IconArea);
+	GetPointHitTestGeometry(InGeometry, WidgetArgs, Points);
+	const FTweenSliderHoverState HoverState = GetHoverState(InMouseEvent.GetScreenSpacePosition(), SliderArea, Points);
+
+	const bool bIsHoveringPoint = HoverState.HoveredPointIndex.IsSet();
+	Pimpl->bStartedMouseDownOnPoint = bIsHoveringPoint;
+	if (bIsHoveringPoint)
+	{
+		// Originally, we had this func call in OnMouseButtonUp (to allow user to move the mouse during click)
+		// but the delay (of user physically releasing mouse button) made it feel unresponsive.
+		HandlePickValueOfCurrentlyHoveredPoint(InGeometry, InMouseEvent);
+		return FReply::Handled().CaptureMouse(SharedThis(this)).PreventThrottling();
+	}
+	
+	Pimpl->DragData.Emplace(InMouseEvent.GetScreenSpacePosition(), InMouseEvent.GetEffectingButton());
+	
+	// If the user clicks in the background (neither slider button nor points), just instantly move the slider there.
+	if (!HoverState.bIsSliderHovered)
+	{
+		Pimpl->DragData->DelayedDrag.ForceDragStart();
+		UpdateLastMousePositionOnSlider(InGeometry, InMouseEvent.GetScreenSpacePosition());
+		Pimpl->TargetSliderPosition = Pimpl->LastMousePositionOnSliderBar;
+		OnSliderDragStartedDelegate.ExecuteIfBound();
+	}
+	
+	return FReply::Handled().CaptureMouse(SharedThis(this)).PreventThrottling();
+}
+
+FReply STweenSlider::OnMouseButtonUp(const FGeometry& InGeometry, const FPointerEvent& InMouseEvent)
+{
+	if (InMouseEvent.GetEffectingButton() != EKeys::LeftMouseButton)
+	{
+		return FReply::Unhandled();
+	}
+
+	if (Pimpl->bStartedMouseDownOnPoint)
+	{
+		Pimpl->bStartedMouseDownOnPoint = false;
+		return FReply::Handled().ReleaseMouseCapture();
+	}
+	
+	if (Pimpl->DragData)
+	{
+		Pimpl->DragData.Reset();
+		OnSliderDragStoppedDelegate.ExecuteIfBound();
+		
+		Pimpl->AnimatedCurrentSliderPosition = 0.5f;
+		Pimpl->TargetSliderPosition = 0.5f;
+		
+		return FReply::Handled()
+			.ReleaseMouseCapture()
+			.ReleaseMouseLock();
+	}
+	
+	return FReply::Unhandled();
+}
+
+void STweenSlider::OnMouseLeave(const FPointerEvent& MouseEvent)
+{
+	Pimpl->HoverState = FTweenSliderHoverState{};
+	Pimpl->bStartedMouseDownOnPoint = false;
+}
+
+FReply STweenSlider::OnMouseMove(const FGeometry& InGeometry, const FPointerEvent& InMouseEvent)
+{
+	const FVector2D MousePosition = InMouseEvent.GetScreenSpacePosition();
+	Pimpl->LastMousePosition = MousePosition;
+	UpdateLastMousePositionOnSlider(InGeometry, MousePosition);
+	
+	if (Pimpl->DragData)
+	{
+		const bool bHasStarted = Pimpl->DragData->DelayedDrag.AttemptDragStart(InMouseEvent);
+		if (bHasStarted)
+		{
+			OnSliderDragStartedDelegate.ExecuteIfBound();
+		}
+
+		if (IsDragging())
+		{
+			Pimpl->TargetSliderPosition = Pimpl->LastMousePositionOnSliderBar;
+		}
+		
+		return FReply::Handled()
+			.CaptureMouse(SharedThis(this))
+			// Locking the mouse gives the user instant feedback about having reached the end of the slider range.
+			.LockMouseToWidget(SharedThis(this));
+	}
+
+	// Hover state needs to updated for OnPaint (don't update it in OnPaint because Slate shifts InGeometry in OnPaint, which messes up hovering computations)
+	FTweenWidgetArgs& WidgetArgs = Pimpl->WidgetArgs;
+	FGeometry SliderArea, IconArea;
+	TArray<FGeometry> Points;
+	GetSliderButtonGeometry(Pimpl->AnimatedCurrentSliderPosition, InGeometry, WidgetArgs, SliderArea, IconArea);
+	GetPointHitTestGeometry(InGeometry, WidgetArgs, Points, Pimpl->PointSliderValues);
+	Pimpl->HoverState = GetHoverState(MousePosition, SliderArea, Points);
+	
+	return Pimpl->bStartedMouseDownOnPoint
+			? FReply::Handled().CaptureMouse(SharedThis(this))
+			: FReply::Unhandled();
+}
+
+void STweenSlider::OnFinishedPointerInput()
+{
+	if (Pimpl->DragData
+		&& Pimpl->DragData->DelayedDrag.IsDragging()
+		&& Pimpl->DragData->LastSliderPosition != Pimpl->TargetSliderPosition)
+	{
+		Pimpl->DragData->LastSliderPosition = Pimpl->TargetSliderPosition;
+
+		// OnSliderChangedDelegate may be an expensive operation!
+		// We call OnSliderChangedDelegate here because OnMouseMove can be called 100s or 1000s of times per tick.
+		// When OnFinishedPointerInput is called, we know the final mouse position for the frame.
+		OnSliderChangedDelegate.ExecuteIfBound(
+			// Convert [0,1] range to [-1,1] range.
+			SliderValueToBlendValue(Pimpl->TargetSliderPosition)
+			);
+	}
+}
+
+FVector2D STweenSlider::ComputeDesiredSize(float) const
+{
+	FTweenWidgetArgs& SliderStyle = Pimpl->WidgetArgs;
+	const FSlateBrush* SliderIconBrush = SliderStyle.SliderIconAttr.Get();
+	
+	const FVector2D TotalButtonSize = SliderIconBrush
+		? SliderIconBrush->ImageSize + SliderStyle.Style->IconPadding.GetDesiredSize()
+		: FVector2D::ZeroVector;
+	
+	return FVector2D
+	{
+		// We want the button to fit at position 0 and 1 - so we need half the button's space each side.
+		SliderStyle.Style->BarDimensions.X + TotalButtonSize.X, 
+		FMath::Max(TotalButtonSize.Y, SliderStyle.Style->BarDimensions.Y)
+	};
+}
+	
+void STweenSlider::Tick(const FGeometry& AllottedGeometry, const double InCurrentTime, const float InDeltaTime)
+{
+	constexpr float InterpSpeed = 35.f;
+	Pimpl->AnimatedCurrentSliderPosition = FMath::FInterpTo(
+		Pimpl->AnimatedCurrentSliderPosition, Pimpl->TargetSliderPosition, InDeltaTime, InterpSpeed
+		);
+}
+
+int32 STweenSlider::OnPaint(
+	const FPaintArgs& InArgs,
+	const FGeometry& InAllottedGeometry,
+	const FSlateRect& InCullingRect,
+	FSlateWindowElementList& OutDrawElements,
+	int32 InLayerId,
+	const FWidgetStyle& InWidgetStyle,
+	bool bParentEnabled
+	) const
+{
+	const bool bIsMouseButtonDown = IsMouseDown();
+	const bool bIsDragging = IsDragging();
+	FTweenWidgetArgs& WidgetArgs = Pimpl->WidgetArgs;
+	FTweenSliderHoverState& HoverState = Pimpl->HoverState;
+	
+	TArray<float> NormalizedPositions;
+	FTweenSliderDrawArgs DrawArgs;
+	
+	GetBarGeometry(InAllottedGeometry, WidgetArgs, DrawArgs.BarArea);
+	GetSliderButtonGeometry(Pimpl->AnimatedCurrentSliderPosition, InAllottedGeometry, WidgetArgs, DrawArgs.SliderArea, DrawArgs.IconArea);
+	GetDrawnPointGeometry(
+		InAllottedGeometry, WidgetArgs, HoverState, bIsMouseButtonDown, DrawArgs.Points, DrawArgs.PointTypes, NormalizedPositions
+		);
+	
+	DrawArgs.bDrawButtonPressed = bIsDragging
+		|| (bIsMouseButtonDown && HoverState.bIsSliderHovered)
+		|| (bIsMouseButtonDown && HoverState.HoveredPointIndex);
+	DrawArgs.bIsDragging = bIsDragging;
+	DrawArgs.HoverState = HoverState;
+	
+	if (bIsDragging)
+	{
+		GetPassedPointStates(NormalizedPositions, Pimpl->AnimatedCurrentSliderPosition, DrawArgs.PassedPoints);
+		GetDragValueIndicationGeometry(InAllottedGeometry, WidgetArgs, Pimpl->AnimatedCurrentSliderPosition, DrawArgs.DragValueIndication);
+	}
+	
+	return DrawTweenSlider(DrawArgs, WidgetArgs, OutDrawElements, InLayerId, InWidgetStyle);
+}
+
+TSharedPtr<IToolTip> STweenSlider::GetToolTip()
+{
+	return Pimpl->ToolTip;
+}
+
+void STweenSlider::UpdateLastMousePositionOnSlider(const FGeometry& InGeometry, const FVector2D& InMousePos)
+{
+	// Compute the offset the widget has from the left side of the screen, subtracts that offset from mouse x, and divides mouse x by widget length.
+	const float DistToLeftScreenEdge = InGeometry.GetAbsolutePosition().X;
+	const float SliderAbsLength = InGeometry.GetAbsoluteSize().X;
+	const float MouseInRange = FMath::Clamp(InMousePos.X - DistToLeftScreenEdge, 0, SliderAbsLength);
+	Pimpl->LastMousePositionOnSliderBar = MouseInRange / SliderAbsLength;
+}
+
+bool STweenSlider::IsMouseDown() const
+{
+	return Pimpl->DragData.IsSet() || Pimpl->bStartedMouseDownOnPoint;
+}
+
+bool STweenSlider::IsDragging() const
+{
+	return Pimpl->DragData && Pimpl->DragData->DelayedDrag.IsDragging(); 
+}
+
+void STweenSlider::HandlePickValueOfCurrentlyHoveredPoint(const FGeometry& InGeometry, const FPointerEvent& InMouseEvent) const
+{
+	FTweenWidgetArgs& WidgetArgs = Pimpl->WidgetArgs;
+	FGeometry SliderArea, IconArea;
+	TArray<FGeometry> Points;
+	TArray<float> PointValues;
+	
+	GetSliderButtonGeometry(Pimpl->AnimatedCurrentSliderPosition, InGeometry, WidgetArgs, SliderArea, IconArea);
+	GetPointHitTestGeometry(InGeometry, WidgetArgs, Points, PointValues);
+	const FTweenSliderHoverState HoverState = GetHoverState(InMouseEvent.GetScreenSpacePosition(), SliderArea, Points);
+
+	const int32 HoverIndex = HoverState.HoveredPointIndex.Get(INDEX_NONE);
+	if (PointValues.IsValidIndex(HoverIndex))
+	{
+		OnPointValuePickedDelegate.ExecuteIfBound(
+			SliderValueToBlendValue(PointValues[HoverIndex])
+			);
+	}
+}
+
+FText STweenSlider::GetToolTipText() const
+{
+	const auto ConvertSliderValue = [this](float Value)
+	{
+		// While sliding the button, the mouse is locked to the widget space (see LockMouseToWidget use above).
+		// However, the mouse is missing a couple pixels on the left and right to fully reach -1 and 1 (it reaches e.g. ~0.978).
+		// Since we display the value in the tooltip, we want the value to be rounded: It'd be weird if you fully move the slider all the way,
+		// and it does not say 1; effectively, we lie to the user, but they won't be able to tell the difference.
+		constexpr float RoundThreshold = 0.03;
+		const bool bRoundLower = Value < RoundThreshold;
+		const bool bRoundUpper = 1 - Value < RoundThreshold;
+		const float AdjustedValue = bRoundLower ? 0.0 : bRoundUpper ? 1.0 : Value;
+		
+		return MapSliderValueToBlendValueDelegate.Execute(SliderValueToBlendValue(AdjustedValue));
+	};
+	
+	const int32 HoveredIndex = Pimpl->HoverState.HoveredPointIndex.Get(INDEX_NONE);
+	const bool bIsHoveringPoint = Pimpl->PointSliderValues.IsValidIndex(HoveredIndex);
+	if (bIsHoveringPoint)
+	{
+		return FText::Format(
+			LOCTEXT("Tooltip.HoverPointFmt", "{0} blend value.\nClick to blend to this value."),
+			ConvertSliderValue(Pimpl->PointSliderValues[HoveredIndex])
+			);
+	}
+
+	if (Pimpl->HoverState.bIsSliderHovered && !IsDragging())
+	{
+		return LOCTEXT("Tooltip.DragSlider", "Drag to blend values.");
+	}
+
+	return FText::Format(
+		LOCTEXT("Tooltip.CurrentBlendValue", "{0} blend value."),
+		ConvertSliderValue(Pimpl->LastMousePositionOnSliderBar)
+		);
+}
+}
+
+#undef LOCTEXT_NAMESPACE
\ No newline at end of file
diff --git a/Engine/Plugins/Animation/TweeningUtils/Source/TweeningUtilsEditor/Private/Widgets/TweenSliderDrawUtils.cpp b/Engine/Plugins/Animation/TweeningUtils/Source/TweeningUtilsEditor/Private/Widgets/TweenSliderDrawUtils.cpp
new file mode 100644
--- /dev/null
+++ b/Engine/Plugins/Animation/TweeningUtils/Source/TweeningUtilsEditor/Private/Widgets/TweenSliderDrawUtils.cpp
@@ -0,0 +1,395 @@
+﻿// Copyright Epic Games, Inc. All Rights Reserved.
+
+#include "Widgets/TweenSliderDrawUtils.h"
+
+#include "Rendering/DrawElementTypes.h"
+#include "Rendering/RenderingCommon.h"
+#include "Widgets/ETweenScaleMode.h"
+
+namespace UE::TweeningUtilsEditor
+{
+static FVector2D ComputeButtonSize(const FTweenWidgetArgs& InWidgetArgs)
+{
+	const FSlateBrush* SliderIconBrush = InWidgetArgs.SliderIconAttr.Get();
+	const FVector2D IconSize = SliderIconBrush ? SliderIconBrush->ImageSize : FVector2D::ZeroVector;
+	return IconSize + InWidgetArgs.Style->IconPadding.GetDesiredSize();
+}
+
+static FVector2D ComputeBarOffset(const FTweenWidgetArgs& InWidgetArgs)
+{
+	// The entire bar is shifted because we need enough space at position 0 to fit the bar
+	const FVector2D HalfButtonOffset { ComputeButtonSize(InWidgetArgs).X / 2.0, 0.0 };
+	return HalfButtonOffset;
+}
+	
+void GetBarGeometry(const FGeometry& AllottedGeometry, const FTweenWidgetArgs& InWidgetArgs, FGeometry& OutBarArea)
+{
+	const FVector2D& Dimensions = InWidgetArgs.Style->BarDimensions;
+	const FVector2D BarSize = FVector2D{ Dimensions.X, Dimensions.Y };
+	OutBarArea = AllottedGeometry.MakeChild(
+		BarSize, FSlateLayoutTransform(ComputeBarOffset(InWidgetArgs) + FVector2D{ 0.f, AllottedGeometry.Size.Y / 2.f - BarSize.Y / 2.f })
+		);
+}
+	
+static FGeometry MakeGeometryOnBar(
+	float InNormalizedBarPos, const FTweenWidgetArgs& InWidgetArgs, const FVector2D& InSize, const FGeometry& InAllottedGeometry
+	)
+{
+	const FVector2D& Dimensions = InWidgetArgs.Style->BarDimensions;
+	const float HorizontalOffset = Dimensions.X * InNormalizedBarPos - (InSize.X / 2.f);
+	return InAllottedGeometry.MakeChild(
+		InSize,
+			FSlateLayoutTransform(ComputeBarOffset(InWidgetArgs)
+			// InAllottedGeometry.Size.Y / 2.f = vertical center of bar,
+			// InSize.Y / 2.f = offset it needs to have from center of bar
+			+ FVector2D{ HorizontalOffset, InAllottedGeometry.Size.Y / 2.f - InSize.Y / 2.f }
+		)
+	);
+}
+	
+void GetDragValueIndicationGeometry(const FGeometry& InAllottedGeometry, const FTweenWidgetArgs& InWidgetArgs, float InSliderPosition, FGeometry& OutDragValueIndication)
+{
+	const FVector2D& Dimensions = InWidgetArgs.Style->BarDimensions;
+	
+	const float RelativeWidth = InSliderPosition - 0.5f;
+	const float AbsoluteWidth = FMath::Abs(RelativeWidth) * Dimensions.X;
+	const bool bHasDraggedLeft = InSliderPosition < 0.5f; 
+	const float StartPos = bHasDraggedLeft ? InSliderPosition : 0.5f;
+
+	const float HorizontalOffset = Dimensions.X * StartPos;
+	const FVector2D Size { AbsoluteWidth , Dimensions.Y };
+	OutDragValueIndication = InAllottedGeometry.MakeChild(
+		Size,
+		FSlateLayoutTransform(ComputeBarOffset(InWidgetArgs)
+			// InAllottedGeometry.Size.Y / 2.f = vertical center of bar,
+			// InSize.Y / 2.f = offset it needs to have from center of bar
+			+ FVector2D(HorizontalOffset, InAllottedGeometry.Size.Y / 2.f - Size.Y / 2.f)
+		)
+	);
+}
+
+void GetSliderButtonGeometry(
+	float InNormalizedPosition, const FGeometry& InAllottedGeometry, const FTweenWidgetArgs& InWidgetArgs, FGeometry& OutSliderArea, FGeometry& OutIconArea
+	)
+{
+	OutSliderArea = MakeGeometryOnBar(InNormalizedPosition, InWidgetArgs, ComputeButtonSize(InWidgetArgs), InAllottedGeometry);
+
+	const FSlateBrush* SliderIconBrush = InWidgetArgs.SliderIconAttr.Get();
+	OutIconArea = SliderIconBrush
+		? MakeGeometryOnBar(InNormalizedPosition, InWidgetArgs, SliderIconBrush->ImageSize, InAllottedGeometry)
+		: InAllottedGeometry.MakeChild(FVector2D::ZeroVector, FSlateLayoutTransform(FVector2D::ZeroVector));
+}
+	
+template<typename TCallback> requires std::is_invocable_v<TCallback, float/*Position*/, EPointType>
+static void EnumeratePoints(const FTweenWidgetArgs& InWidgetArgs, TCallback&& AddPoint)
+{
+	AddPoint(0.f, EPointType::End); 
+	AddPoint(1.f, EPointType::End); 
+
+	if (InWidgetArgs.ScaleRenderModeAttr.Get() == ETweenScaleMode::Normalized)
+	{
+		// Range is -100% to 100%. Place Small points every 12.5%
+		for (int32 PointIndex = 1; PointIndex < 8; ++PointIndex)
+		{
+			const float Percentage = PointIndex / 8.f;
+			AddPoint(Percentage, EPointType::Small); 
+		}
+	}
+	else
+	{
+		// Range is -200% to 200%. Place Small points every 25%
+		for (int32 PointIndex = 1; PointIndex < 16; ++PointIndex)
+		{
+			const bool bShouldBeMediumPoint = PointIndex == 4 || PointIndex == 12;
+			if (!bShouldBeMediumPoint)
+			{
+				const float Percentage = PointIndex / 16.0;
+				AddPoint(Percentage, EPointType::Small); 
+			}
+		}
+		
+		AddPoint(0.25f, EPointType::Medium); 
+		AddPoint(0.75f, EPointType::Medium); 
+	}
+	
+	AddPoint(0.5f, EPointType::Medium); 
+}
+	
+static const FTweenPointStyle* GetPointStyleFromPointType(const FTweenSliderStyle& InStyle, EPointType PointType)
+{
+	static_assert(static_cast<int32>(EPointType::Num) == 3, "Update this switch");
+	switch (PointType)
+	{
+	case EPointType::Small: return &InStyle.SmallPoint;
+	case EPointType::Medium: return &InStyle.MediumPoint;
+	case EPointType::End: return &InStyle.EndPoint;
+	default: checkNoEntry(); return nullptr;
+	}
+}
+	
+enum class EHoverState
+{
+	Normal,
+	Hovered,
+	Pressed
+};
+static FVector2D GetPointSize(EHoverState HoverState, const FTweenPointStyle& Style)
+{
+	switch (HoverState)
+	{
+	case EHoverState::Normal: return Style.Normal.ImageSize;
+	case EHoverState::Hovered: return Style.Hovered.ImageSize;
+	case EHoverState::Pressed: return Style.Pressed.ImageSize;
+	default: checkNoEntry(); return FVector2d::ZeroVector;
+	}
+}
+	
+void GetDrawnPointGeometry(
+	const FGeometry& InAllottedGeometry, const FTweenWidgetArgs& InWidgetArgs,
+	const FTweenSliderHoverState& InHoverState, bool bIsMouseButtonDown,
+	TArray<FGeometry>& OutPoints, TArray<EPointType>& OutPointTypes, TArray<float>& OutNormalizedPositions
+	)
+{
+	OutPoints.Reset();
+	OutPointTypes.Reset();
+	OutNormalizedPositions.Reset();
+	
+	const int32 HoveredIndex = InHoverState.HoveredPointIndex.Get(INDEX_NONE);
+	EnumeratePoints(InWidgetArgs,
+		[&InAllottedGeometry, &InWidgetArgs, bIsMouseButtonDown, &OutPoints, &OutPointTypes, &OutNormalizedPositions, HoveredIndex]
+		(float InPosition, EPointType PointType)
+		{
+			const int32 AddedPointIndex = OutPoints.Num();
+			const EHoverState PointHoverState =  AddedPointIndex == HoveredIndex
+				? bIsMouseButtonDown ? EHoverState::Hovered : EHoverState::Pressed
+				: EHoverState::Normal;
+
+			const FVector2D Size = GetPointSize(PointHoverState, *GetPointStyleFromPointType(*InWidgetArgs.Style, PointType));
+			const FGeometry Geometry = MakeGeometryOnBar(
+				InPosition, InWidgetArgs, Size, InAllottedGeometry
+				);
+			
+			OutPoints.Add(Geometry);
+			OutPointTypes.Add(PointType);
+			OutNormalizedPositions.Add(InPosition);
+		});
+}
+	
+void GetPointHitTestGeometry(
+	const FGeometry& InAllottedGeometry, const FTweenWidgetArgs& InWidgetArgs, TArray<FGeometry>& OutPoints, TArray<float>& OutPointSliderValues
+	)
+{
+	OutPoints.Reset();
+	OutPointSliderValues.Reset();
+	
+	EnumeratePoints(InWidgetArgs,
+		[&InAllottedGeometry, &InWidgetArgs, &OutPoints, &OutPointSliderValues](float InPosition, EPointType PointType)
+		{
+			const FVector2D& HitTestSize = GetPointStyleFromPointType(*InWidgetArgs.Style, PointType)->HitTestSize;
+			const FGeometry Geometry = MakeGeometryOnBar(
+				InPosition, InWidgetArgs, HitTestSize, InAllottedGeometry
+				);
+			
+			OutPoints.Add(Geometry);
+			OutPointSliderValues.Add(InPosition);
+		});
+}
+
+void GetPointHitTestGeometry(const FGeometry& InAllottedGeometry, const FTweenWidgetArgs& InWidgetArgs, TArray<FGeometry>& OutPoints)
+{
+	OutPoints.Reset();
+	
+	EnumeratePoints(InWidgetArgs,
+		[&InAllottedGeometry, &InWidgetArgs, &OutPoints](float InPosition, EPointType PointType)
+		{
+			const FVector2D& HitTestSize = GetPointStyleFromPointType(*InWidgetArgs.Style, PointType)->HitTestSize;
+			const FGeometry Geometry = MakeGeometryOnBar(
+				InPosition, InWidgetArgs, HitTestSize, InAllottedGeometry
+				);
+			
+			OutPoints.Add(Geometry);
+		});
+}
+
+void GetPassedPointStates(const TArray<float>& InNormalizedPositions, float InSliderPosition, TBitArray<>& OutPassedPoints)
+{
+	const bool bIsDraggingLeft = InSliderPosition < 0.5f;
+	for (const float Position : InNormalizedPositions)
+	{
+		const bool bIsLeft = Position < 0.5f;
+		const bool bIsPassedPoint = (bIsDraggingLeft && bIsLeft && InSliderPosition < Position)
+			|| (!bIsDraggingLeft && !bIsLeft && Position < InSliderPosition)
+			// The center point is always passed.
+			|| FMath::IsNearlyEqual(Position, 0.5f);
+		OutPassedPoints.Add(bIsPassedPoint);
+	}
+}
+	
+FTweenSliderHoverState GetHoverState(const FVector2D& InMouseScreenSpace, const FGeometry& InButtonArea, const TArray<FGeometry>& InPoints)
+{
+	FTweenSliderHoverState Result;
+
+	if (InButtonArea.IsUnderLocation(InMouseScreenSpace))
+	{
+		Result.bIsSliderHovered = true;
+	}
+	else
+	{
+		for (int32 PointIndex = 0; PointIndex < InPoints.Num(); ++PointIndex)
+		{
+			if (InPoints[PointIndex].IsUnderLocation(InMouseScreenSpace))
+			{
+				Result.HoveredPointIndex = PointIndex;
+				break;
+			}
+		}
+	}
+	
+	return Result;
+}
+
+static void DrawBar(
+	const FTweenSliderDrawArgs& InDrawArgs,
+	const FTweenWidgetArgs& InWidgetArgs,
+	FSlateWindowElementList& OutDrawElements,
+	int32 LayerId,
+	const FWidgetStyle& InWidgetStyle
+	)
+{
+	const FTweenSliderStyle* Style = InWidgetArgs.Style;
+	
+	const FSlateBrush* BarBrush = &Style->BarBrush;
+	const FLinearColor FinalBackgroundColor = InWidgetArgs.ColorAndOpacity.Get().GetColor(InWidgetStyle) * BarBrush->TintColor.GetColor(InWidgetStyle);
+	FSlateDrawElement::MakeBox(
+		OutDrawElements, LayerId, InDrawArgs.BarArea.ToPaintGeometry(), BarBrush, ESlateDrawEffect::None, FinalBackgroundColor
+		);
+}
+	
+static void DrawDragValueIndication(
+	const FTweenSliderDrawArgs& InDrawArgs,
+	const FTweenWidgetArgs& InWidgetArgs,
+	FSlateWindowElementList& OutDrawElements,
+	int32 LayerId,
+	const FWidgetStyle& InWidgetStyle
+	)
+{
+	if (!InDrawArgs.bIsDragging)
+	{
+		return;
+	}
+
+	const FTweenSliderStyle* Style = InWidgetArgs.Style;
+	const FLinearColor BarColor =
+		InWidgetArgs.ColorAndOpacity.Get().GetColor(InWidgetStyle)
+		* InWidgetArgs.SliderColor.Get()
+		* Style->PassedValueBackground.TintColor.GetColor(InWidgetStyle);
+	FSlateDrawElement::MakeBox(
+		OutDrawElements, LayerId, InDrawArgs.DragValueIndication.ToPaintGeometry(), &Style->PassedValueBackground, ESlateDrawEffect::None, BarColor
+		);
+}
+
+/** Should the point's color be multiplied with the color of the slider? */
+enum class EUseSliderColor { Yes, No };
+	
+static TPair<const FSlateBrush*, EUseSliderColor> GetPointBrush(
+	const FTweenSliderDrawArgs& InDrawArgs, const FTweenSliderStyle& Style, int32 Index
+	)
+{
+	const FTweenPointStyle* PointStyle = GetPointStyleFromPointType(Style, InDrawArgs.PointTypes[Index]);
+	
+	const bool bIsPassedPoint = InDrawArgs.PassedPoints.IsValidIndex(Index) && InDrawArgs.PassedPoints[Index];
+	if (bIsPassedPoint)
+	{
+		// Passed points are white
+		return { &PointStyle->PassedPoint, EUseSliderColor::No };
+	}
+	
+	if (InDrawArgs.bIsDragging || InDrawArgs.HoverState.HoveredPointIndex != Index)
+	{
+		return { &PointStyle->Normal , EUseSliderColor::Yes };
+	}
+	
+	if (InDrawArgs.bDrawButtonPressed)
+	{
+		// Pressed points are white
+		return { &PointStyle->Pressed, EUseSliderColor::No };
+	}
+	return { &PointStyle->Hovered, EUseSliderColor::Yes };
+}
+
+static void DrawBarPoints(
+	const FTweenSliderDrawArgs& InDrawArgs,
+	const FTweenWidgetArgs& InWidgetArgs,
+	FSlateWindowElementList& OutDrawElements,
+	int32 LayerId,
+	const FWidgetStyle& InWidgetStyle
+	)
+{
+	const FTweenSliderStyle* Style = InWidgetArgs.Style;
+	const FLinearColor ColorAndOpacity = InWidgetArgs.ColorAndOpacity.Get().GetColor(InWidgetStyle);
+	const FLinearColor SliderColor = InWidgetArgs.SliderColor.Get();
+	
+	for (int32 PointIndex = 0; PointIndex < InDrawArgs.Points.Num(); ++PointIndex)
+	{
+		const FGeometry& PointGeometry = InDrawArgs.Points[PointIndex];
+		const auto[Brush, UseSliderColor] = GetPointBrush(InDrawArgs, *Style, PointIndex);
+		
+		const FLinearColor PointColor = UseSliderColor == EUseSliderColor::Yes
+			? ColorAndOpacity * SliderColor * Brush->TintColor.GetColor(InWidgetStyle)
+			: ColorAndOpacity * Brush->TintColor.GetColor(InWidgetStyle);
+		FSlateDrawElement::MakeBox(
+			OutDrawElements, LayerId, PointGeometry.ToPaintGeometry(), Brush, ESlateDrawEffect::None, PointColor
+		);
+	}
+}
+
+static void DrawButton(
+	const FTweenSliderDrawArgs& InDrawArgs,
+	const FTweenWidgetArgs& InWidgetArgs,
+	FSlateWindowElementList& OutDrawElements,
+	int32 LayerId,
+	const FWidgetStyle& InWidgetStyle
+	)
+{
+	const FTweenSliderStyle* Style = InWidgetArgs.Style;
+	
+	const FSlateBrush* SliderIconBrush = InWidgetArgs.SliderIconAttr.Get();
+	const FSlateBrush* SliderButtonBrush = InDrawArgs.bDrawButtonPressed
+		? &Style->PressedSliderButton
+		: InDrawArgs.HoverState.bIsSliderHovered ? &Style->HoveredSliderButton : &Style->NormalSliderButton;
+
+	const FSlateColor& BaseIconTint = InDrawArgs.bDrawButtonPressed
+		? Style->PressedIconTint
+		: InDrawArgs.HoverState.bIsSliderHovered ? Style->PressedIconTint : Style->PressedIconTint;
+	
+	const FLinearColor ColorAndOpacity = InWidgetArgs.ColorAndOpacity.Get().GetColor(InWidgetStyle);
+	const FLinearColor IconTint = ColorAndOpacity * BaseIconTint.GetColor(InWidgetStyle) * SliderIconBrush->TintColor.GetColor(InWidgetStyle);
+	const FLinearColor ButtonTint = ColorAndOpacity * InWidgetArgs.SliderColor.Get() * SliderButtonBrush->TintColor.GetColor(InWidgetStyle);
+	
+	FSlateDrawElement::MakeBox(
+		OutDrawElements, LayerId, InDrawArgs.SliderArea.ToPaintGeometry(), SliderButtonBrush, ESlateDrawEffect::None, ButtonTint
+		);
+	if (SliderIconBrush)
+	{
+		FSlateDrawElement::MakeBox(
+			OutDrawElements, LayerId + 1, InDrawArgs.IconArea.ToPaintGeometry(), SliderIconBrush, ESlateDrawEffect::None, IconTint
+		);
+	}
+}
+
+int32 DrawTweenSlider(
+	const FTweenSliderDrawArgs& InDrawArgs,
+	const FTweenWidgetArgs& InWidgetArgs,
+	FSlateWindowElementList& OutDrawElements,
+	int32 LayerId,
+	const FWidgetStyle& InWidgetStyle
+	)
+{
+	DrawBar(InDrawArgs, InWidgetArgs, OutDrawElements, LayerId, InWidgetStyle);
+	DrawDragValueIndication(InDrawArgs, InWidgetArgs, OutDrawElements, LayerId, InWidgetStyle);
+	DrawBarPoints(InDrawArgs, InWidgetArgs, OutDrawElements, LayerId, InWidgetStyle);
+	DrawButton(InDrawArgs, InWidgetArgs, OutDrawElements, LayerId, InWidgetStyle);
+	
+	return LayerId;
+}
+}
diff --git a/Engine/Plugins/Animation/TweeningUtils/Source/TweeningUtilsEditor/Private/Widgets/TweenSliderDrawUtils.h b/Engine/Plugins/Animation/TweeningUtils/Source/TweeningUtilsEditor/Private/Widgets/TweenSliderDrawUtils.h
new file mode 100644
--- /dev/null
+++ b/Engine/Plugins/Animation/TweeningUtils/Source/TweeningUtilsEditor/Private/Widgets/TweenSliderDrawUtils.h
@@ -0,0 +1,136 @@
+﻿// Copyright Epic Games, Inc. All Rights Reserved.
+
+#pragma once
+
+#include "Containers/Array.h"
+#include "Containers/BitArray.h"
+#include "HAL/Platform.h"
+#include "Layout/Geometry.h"
+#include "Math/Color.h"
+#include "Misc/Attribute.h"
+#include "Styling/SlateColor.h"
+#include "Widgets/TweenSliderStyle.h"
+
+class FSlateWindowElementList;
+class FWidgetStyle;
+namespace UE::TweeningUtilsEditor { enum class ETweenScaleMode : uint8; }
+
+namespace UE::TweeningUtilsEditor
+{	
+struct FTweenSliderHoverState
+{
+	/** Whether the slider button is pressed. */
+	bool bIsSliderHovered = false;
+
+	/** Index of FTweenSliderDrawArgs::Points that is hovered. */
+	TOptional<int32> HoveredPointIndex;
+};
+
+enum class EPointType : uint8
+{
+	Small,
+	/** When overshoot mode is enabled, a bar at 100% and -100% */
+	Medium,
+	/** Left or right end */
+	End,
+
+	Num
+};
+	
+/**
+ * Contains the basic primitives that are supposed to be drawn. This is designed to keep the drawing algorithm as straight forward as possible.
+ * This way, theoretically, we could unit test the geometry generation.
+ */
+struct FTweenSliderDrawArgs
+{
+	/** The black bar */
+	FGeometry BarArea;
+	
+	/** The points on the bar */
+	TArray<FGeometry> Points;
+	/** Equal length as Points. Used to get the right brush. */
+	TArray<EPointType> PointTypes;
+	/**
+	 * Indicates which of the points should be rendered as passed (i.e. the slider has moved over them): true if passed, false if not passed.
+	 * Equal length as Points when dragging. 0 length if not dragging. 
+	 */
+	TBitArray<> PassedPoints;
+	
+	/** The slider button. */
+	FGeometry SliderArea;
+	/** The area of icon in the button */
+	FGeometry IconArea;
+
+	/** This widget is shown while dragging: it draws from the center to the slider. It helps the user see how much they dragged so far. */
+	FGeometry DragValueIndication;
+	
+	/** Whether to draw the button as pressed down. Can be true when bIsDragging is false (e.g. while detecting a drag). */
+	bool bDrawButtonPressed = false;
+	/** Whether the slider is being dragged. Affects whether to draw DragValueIndication and how to draw points. */
+	bool bIsDragging = false;
+	
+	/** Indicates the hovered element */
+	FTweenSliderHoverState HoverState;
+};
+
+/** The drawing relevant construction args of the widget. */
+struct FTweenWidgetArgs
+{
+	/** The style that was used to construct the widget. */
+	const FTweenSliderStyle* Style;
+	
+	/** The root opacity of the widget. */
+	TAttribute<FSlateColor> ColorAndOpacity;
+
+	/** The icon to place in the button */
+	TAttribute<const FSlateBrush*> SliderIconAttr;
+	/** The main color. It tints the slider button and the points. */
+	TAttribute<FLinearColor> SliderColor;
+	/** Affects how the normalized values are supposed to be interpreted */
+	TAttribute<ETweenScaleMode> ScaleRenderModeAttr;
+};
+
+/** Gets geometry of the background bar. */
+void GetBarGeometry(const FGeometry& AllottedGeometry, const FTweenWidgetArgs& InWidgetArgs, FGeometry& OutBarArea);
+/** This widget is shown while dragging: it draws from the center to the slider. */
+void GetDragValueIndicationGeometry(
+	const FGeometry& InAllottedGeometry, const FTweenWidgetArgs& InWidgetArgs, float InSliderPosition, FGeometry& OutDragValueIndication
+	);
+/** Computes the geometry for the slider button */
+void GetSliderButtonGeometry(
+	float InNormalizedPosition, const FGeometry& InAllottedGeometry, const FTweenWidgetArgs& InWidgetArgs, FGeometry& OutSliderArea, FGeometry& OutIconArea
+	);
+	
+/** Gets the drawn geometry of the points on the bar. */
+void GetDrawnPointGeometry(
+	const FGeometry& InAllottedGeometry, const FTweenWidgetArgs& InWidgetArgs,
+	const FTweenSliderHoverState& InHoverState, bool bIsMouseButtonDown,
+	TArray<FGeometry>& OutPoints, TArray<EPointType>& OutPointTypes, TArray<float>& OutNormalizedPositions
+	);
+/** Gets the point geometry for doing hit tests with. */
+void GetPointHitTestGeometry(
+	const FGeometry& InAllottedGeometry, const FTweenWidgetArgs& InWidgetArgs,
+	TArray<FGeometry>& OutPoints, TArray<float>& OutPointSliderValues
+	);
+/** Util overload for when the caller does not care about the point values. */
+void GetPointHitTestGeometry(
+	const FGeometry& InAllottedGeometry, const FTweenWidgetArgs& InWidgetArgs, TArray<FGeometry>& OutPoints 
+);
+
+/** Computes the points that have been passed by the slider while dragging. */
+void GetPassedPointStates(const TArray<float>& InNormalizedPositions, float InSliderPosition, TBitArray<>& OutPassedPoints);
+/** Computes the hover state. */
+FTweenSliderHoverState GetHoverState(const FVector2D& InMouseScreenSpace, const FGeometry& InButtonArea, const TArray<FGeometry>& InPoints);
+	
+/**
+ * Draws a tween slider.
+ * @return The max LayerId that was drawn onto.
+ */
+int32 DrawTweenSlider(
+	const FTweenSliderDrawArgs& InDrawArgs,
+	const FTweenWidgetArgs& InWidgetArgs,
+	FSlateWindowElementList& OutDrawElements,
+	int32 LayerId,
+	const FWidgetStyle& InWidgetStyle
+	);
+}
diff --git a/Engine/Plugins/Animation/TweeningUtils/Source/TweeningUtilsEditor/Private/Widgets/TweenSliderStyle.cpp b/Engine/Plugins/Animation/TweeningUtils/Source/TweeningUtilsEditor/Private/Widgets/TweenSliderStyle.cpp
new file mode 100644
--- /dev/null
+++ b/Engine/Plugins/Animation/TweeningUtils/Source/TweeningUtilsEditor/Private/Widgets/TweenSliderStyle.cpp
@@ -0,0 +1,66 @@
+﻿// Copyright Epic Games, Inc. All Rights Reserved.
+
+#include "Widgets/TweenSliderStyle.h"
+
+#include "Brushes/SlateColorBrush.h"
+#include "Brushes/SlateImageBrush.h"
+#include "Brushes/SlateRoundedBoxBrush.h"
+#include "Styling/SlateColor.h"
+#include "Styling/StyleColors.h"
+
+const FName FTweenSliderStyle::TypeName( TEXT("FTweenSliderStyle") );
+
+FTweenPointStyle::FTweenPointStyle() = default;
+
+FTweenPointStyle::FTweenPointStyle(
+	const FVector2D& InNormalSize, const FVector2D& InHoveredSize, const FVector2D& InPressedSize, const FVector2D& InHitSize
+	)
+	: Normal(static_cast<FSlateBrush>(FSlateImageBrush(NAME_None, InNormalSize, FStyleColors::ForegroundHover)))
+	, Hovered(static_cast<FSlateBrush>(FSlateImageBrush(NAME_None, InHoveredSize, FStyleColors::ForegroundHover)))
+	, Pressed(static_cast<FSlateBrush>(FSlateImageBrush(NAME_None, InPressedSize, FStyleColors::White)))
+	, PassedPoint(static_cast<FSlateBrush>(FSlateImageBrush(NAME_None, InNormalSize, FStyleColors::White)))
+	, HitTestSize(InHitSize)
+{}
+
+void FTweenPointStyle::GetResources(TArray<const FSlateBrush*>& OutBrushes) const
+{
+	OutBrushes.Add(&Normal);
+	OutBrushes.Add(&Hovered);
+	OutBrushes.Add(&Pressed);
+	OutBrushes.Add(&PassedPoint);
+}
+
+FTweenSliderStyle::FTweenSliderStyle()
+	: BarDimensions(250.0, 12.0)
+	, BarBrush(static_cast<FSlateBrush>(FSlateColorBrush(FStyleColors::Black)))
+	, NormalSliderButton(static_cast<FSlateBrush>(FSlateRoundedBoxBrush(FStyleColors::Foreground, 5.0f, FStyleColors::Transparent, 0.f)))
+	, HoveredSliderButton(static_cast<FSlateBrush>(FSlateRoundedBoxBrush(FStyleColors::ForegroundHover, 5.0f, FStyleColors::Transparent, 0.f)))
+	, PressedSliderButton(static_cast<FSlateBrush>(FSlateRoundedBoxBrush(FStyleColors::Foreground, 5.0f, FStyleColors::Transparent, 0.f)))
+	, NormalIconTint(FStyleColors::Foreground)
+	, HoveredIconTint(FStyleColors::White)
+	, PressedIconTint(FStyleColors::White)
+	, SmallPoint({ 4.0, 4.0 }, { 6.0, 6.0 }, { 6.0, 6.0 }, { 10.0, 12.0 })
+	, MediumPoint({ 4.0, 8.0 }, { 6.0, 10.0 }, { 6.0, 10.0 }, { 10.0, 12.0 })
+	, EndPoint({ 4.0, 12.0 }, { 6.0, 14.0 }, { 6.0, 14.0 }, { 10.0, 12.0 })
+	, PassedValueBackground(static_cast<FSlateBrush>(FSlateColorBrush(FLinearColor(1.f, 1.f, 1.f, 0.3f))))
+	, IconPadding(4.f, 2.f)
+{
+}
+
+void FTweenSliderStyle::GetResources(TArray<const FSlateBrush*>& OutBrushes) const
+{
+	OutBrushes.Add(&BarBrush);
+	OutBrushes.Add(&NormalSliderButton);
+	OutBrushes.Add(&HoveredSliderButton);
+	OutBrushes.Add(&PressedSliderButton);
+	OutBrushes.Add(&PassedValueBackground);
+	SmallPoint.GetResources(OutBrushes);
+	MediumPoint.GetResources(OutBrushes);
+	EndPoint.GetResources(OutBrushes);
+}
+
+const FTweenSliderStyle& FTweenSliderStyle::GetDefault()
+{
+	static FTweenSliderStyle Default;
+	return Default;
+}
diff --git a/Engine/Plugins/Animation/TweeningUtils/Source/TweeningUtilsEditor/Public/Math/Abstraction/BuiltInTweenFunction.h b/Engine/Plugins/Animation/TweeningUtils/Source/TweeningUtilsEditor/Public/Math/Abstraction/BuiltInTweenFunction.h
new file mode 100644
--- /dev/null
+++ b/Engine/Plugins/Animation/TweeningUtils/Source/TweeningUtilsEditor/Public/Math/Abstraction/BuiltInTweenFunction.h
@@ -0,0 +1,29 @@
+﻿// Copyright Epic Games, Inc. All Rights Reserved.
+
+#pragma once
+
+#include "ITweenFunction.h"
+#include "KeyBlendingAbstraction.h"
+
+namespace UE::TweeningUtilsEditor
+{
+
+/** Evaluates a blend function that is identified by EBlendFunction. */
+class FBuiltInTweenFunction : public ITweenFunction
+{
+public:
+
+	explicit FBuiltInTweenFunction(EBlendFunction InBlendFunction) : BlendFunction(InBlendFunction) {}
+
+	EBlendFunction GetBlendFunction() const { return BlendFunction; }
+	void SetBlendFunction(EBlendFunction InValue) { BlendFunction = InValue; }
+	
+	//~ Begin ITweenFunction Interface
+	virtual double Blend(const FKeyBlendInfo& InBlendInfo) const override { return EvaluateBlend(BlendFunction, InBlendInfo); }
+	//~ End ITweenFunction Interface
+
+private:
+
+	EBlendFunction BlendFunction;
+};
+}
diff --git a/Engine/Plugins/Animation/TweeningUtils/Source/TweeningUtilsEditor/Public/Math/Abstraction/BuiltInTweenFunctionContainer.h b/Engine/Plugins/Animation/TweeningUtils/Source/TweeningUtilsEditor/Public/Math/Abstraction/BuiltInTweenFunctionContainer.h
new file mode 100644
--- /dev/null
+++ b/Engine/Plugins/Animation/TweeningUtils/Source/TweeningUtilsEditor/Public/Math/Abstraction/BuiltInTweenFunctionContainer.h
@@ -0,0 +1,36 @@
+﻿// Copyright Epic Games, Inc. All Rights Reserved.
+
+#pragma once
+
+#include "Containers/Array.h"
+#include "ITweenFunctionContainer.h"
+#include "Templates/SharedPointer.h"
+
+namespace UE::TweeningUtilsEditor
+{
+class FBuiltInTweenFunction;
+
+/** Contains all the functions from this module, i.e. those specified by EBlendFunction.  */
+class TWEENINGUTILSEDITOR_API FBuiltInTweenFunctionContainer : public ITweenFunctionContainer
+{
+public:
+
+	FBuiltInTweenFunctionContainer();
+
+	//~ Begin ITweenFunctionContainer Interface
+	virtual void ForEachFunction(const TFunctionRef<void(const TSharedRef<ITweenFunction>&)> InConsumer) const override;
+	virtual TSharedPtr<ITweenFunction> GetFunction(int32 InIndex) const override;
+	virtual int32 NumFunctions() const override { return Functions.Num(); }
+	virtual TSharedPtr<FUICommandInfo> GetCommandForFunction(const TSharedRef<ITweenFunction>& InFunction) const override;
+	virtual const FSlateBrush* GetIconForFunction(const TSharedRef<ITweenFunction>& InFunction) const override;
+	virtual FLinearColor GetColorForFunction(const TSharedRef<ITweenFunction>& InFunction) const override;
+	virtual FText GetLabelForFunction(const TSharedRef<ITweenFunction>& InFunction) const override;
+	virtual FText GetToolTipForFunction(const TSharedRef<ITweenFunction>& InFunction) const override;
+	//~ Begin ITweenFunctionContainer Interface
+
+private:
+
+	/** The built-in functions. */
+	TArray<TSharedRef<FBuiltInTweenFunction>> Functions;
+};
+}
diff --git a/Engine/Plugins/Animation/TweeningUtils/Source/TweeningUtilsEditor/Public/Math/Abstraction/ITweenFunction.h b/Engine/Plugins/Animation/TweeningUtils/Source/TweeningUtilsEditor/Public/Math/Abstraction/ITweenFunction.h
new file mode 100644
--- /dev/null
+++ b/Engine/Plugins/Animation/TweeningUtils/Source/TweeningUtilsEditor/Public/Math/Abstraction/ITweenFunction.h
@@ -0,0 +1,19 @@
+﻿// Copyright Epic Games, Inc. All Rights Reserved.
+
+#pragma once
+
+namespace UE::TweeningUtilsEditor
+{
+struct FKeyBlendInfo;
+
+/** Abstracts tweening functions. This is useful for maintaining generic UI but support adding custom functions from external modules. */
+class ITweenFunction
+{
+public:
+
+	/** @return The result of combining FKeyBlendInfo::CurrentValue and FKeyBlendInfo::BlendValue according to this function. */
+	virtual double Blend(const FKeyBlendInfo& InBlendInfo) const = 0;
+
+	virtual ~ITweenFunction() = default;
+};
+}
\ No newline at end of file
diff --git a/Engine/Plugins/Animation/TweeningUtils/Source/TweeningUtilsEditor/Public/Math/Abstraction/ITweenFunctionContainer.h b/Engine/Plugins/Animation/TweeningUtils/Source/TweeningUtilsEditor/Public/Math/Abstraction/ITweenFunctionContainer.h
new file mode 100644
--- /dev/null
+++ b/Engine/Plugins/Animation/TweeningUtils/Source/TweeningUtilsEditor/Public/Math/Abstraction/ITweenFunctionContainer.h
@@ -0,0 +1,45 @@
+﻿// Copyright Epic Games, Inc. All Rights Reserved.
+
+#pragma once
+
+#include "Templates/FunctionFwd.h"
+#include "Templates/SharedPointerFwd.h"
+
+class FText;
+class FUICommandInfo;
+struct FLinearColor;
+struct FSlateBrush;
+
+namespace UE::TweeningUtilsEditor
+{
+class ITweenFunction;
+
+/** This acts as the source of functions displayed in UI, e.g. by FTweenToolbarController. */
+class ITweenFunctionContainer
+{
+public:
+
+	/** Iterates through all the functions */
+	virtual void ForEachFunction(const TFunctionRef<void(const TSharedRef<ITweenFunction>&)> InConsumer) const = 0;
+
+	/** @return The function at InIndex. */
+	virtual TSharedPtr<ITweenFunction> GetFunction(int32 InIndex) const = 0;
+	/** @return The number of functions contained. */
+	virtual int32 NumFunctions() const = 0;
+
+	/** @return The command that is used to select the function */
+	virtual TSharedPtr<FUICommandInfo> GetCommandForFunction(const TSharedRef<ITweenFunction>& InFunction) const = 0;
+
+	/** @return The un-tinted icon to display the function in the UI with */
+	virtual const FSlateBrush* GetIconForFunction(const TSharedRef<ITweenFunction>& InFunction) const = 0;
+	/** @return The color that represents the function in the UI */
+	virtual FLinearColor GetColorForFunction(const TSharedRef<ITweenFunction>& InFunction) const = 0;
+
+	/** @return The label to display in the combo button */
+	virtual FText GetLabelForFunction(const TSharedRef<ITweenFunction>& InFunction) const = 0;
+	/** @return The description to displayed, e.g. in tool tips. */
+	virtual FText GetToolTipForFunction(const TSharedRef<ITweenFunction>& InFunction) const = 0;
+
+	virtual ~ITweenFunctionContainer() = default;
+};
+}
diff --git a/Engine/Plugins/Animation/TweeningUtils/Source/TweeningUtilsEditor/Public/Math/Abstraction/KeyBlendingAbstraction.h b/Engine/Plugins/Animation/TweeningUtils/Source/TweeningUtilsEditor/Public/Math/Abstraction/KeyBlendingAbstraction.h
new file mode 100644
--- /dev/null
+++ b/Engine/Plugins/Animation/TweeningUtils/Source/TweeningUtilsEditor/Public/Math/Abstraction/KeyBlendingAbstraction.h
@@ -0,0 +1,93 @@
+﻿// Copyright Epic Games, Inc. All Rights Reserved.
+
+#pragma once
+
+#include "HAL/Platform.h"
+#include "Templates/SharedPointerFwd.h"
+
+#include <type_traits>
+
+class FText;
+class FUICommandInfo;
+struct FLinearColor;
+struct FSlateBrush;
+
+/** The following just contains some utils that is common in UI code. Tools can opt for their own description but using this helps with consistency. */
+namespace UE::TweeningUtilsEditor
+{
+struct FKeyBlendInfo;
+	
+/** Util enum which is useful for UI code to abstract functions.*/
+enum class EBlendFunction : uint8
+{
+	ControlsToTween,
+	PushPull,
+	BlendNeighbor,
+	BlendRelative,
+	BlendEase,
+	SmoothRough,
+
+	// Steps for adding a new function:
+	// 1. Add the function to KeyBlendingFunctions.h (consider splitting implementation to KeyBlendingFunctions_[YourFunc].cpp).
+	// 2. Extend this enum
+	// 3. Extend FTweeningUtilsCommands with another command
+	// 4. Extend FTweeningUtilsStyle with color, command style, and icon.
+	// 5. Update GetFunctionData in KeyBlendingAbstraction.cpp
+	// Done. Every system using the below functions will be up to date.
+	//
+	// Once you extend this  enum, you cannot miss the required steps as they're all guarded by static_asserts. If it compiles, you did everything.
+	// Make sure you add static_assert(static_cast<int32>(EBlendFunction::Num) == x, "Instruction");, if you add additional steps in the future!
+
+	// Add new entries above.
+	Num
+};
+
+/** Evaluates the function */
+TWEENINGUTILSEDITOR_API double EvaluateBlend(EBlendFunction InFunction, const FKeyBlendInfo& InBlendInfo);
+
+/** @return The command that is used to select the function in most UI. */
+TWEENINGUTILSEDITOR_API TSharedPtr<FUICommandInfo> GetCommandForBlendFunction(EBlendFunction InFunction);
+/** @return The un-tinted icon to display the function in the UI with */
+TWEENINGUTILSEDITOR_API	const FSlateBrush* GetUntintedIconForTweenFunction(EBlendFunction InFunction);
+/** @return The color that represents the function in the UI */
+TWEENINGUTILSEDITOR_API FLinearColor GetTintColorForTweenFunction(EBlendFunction InFunction);
+	
+/** @return The full label to display in the function with, e.g. "Push / Pull" */
+TWEENINGUTILSEDITOR_API	FText GetLabelForBlendFunction(EBlendFunction InFunction);
+/** @return The short label to display in the function with, e.g. "PP" (for Push / Pull) */
+TWEENINGUTILSEDITOR_API	FText GetAbbreviationForBlendFunction(EBlendFunction InFunction);
+/** @return The description of the function */
+TWEENINGUTILSEDITOR_API	FText GetDescriptionForBlendFunction(EBlendFunction InFunction);
+
+/** Iterates through all blend function types. */
+template<typename TCallback> requires std::is_invocable_v<TCallback, EBlendFunction>
+void ForEachBlendFunction(TCallback&& InCallback);
+	
+enum class EBreakBehavior : uint8 { Continue, Break };
+/** Iterates through all blend function types with the ability to break. */
+template<typename TCallback> requires std::is_invocable_r_v<EBreakBehavior, TCallback, EBlendFunction>
+void ForEachBlendFunctionBreakable(TCallback&& InCallback);
+}
+
+namespace UE::TweeningUtilsEditor
+{
+template<typename TCallback> requires std::is_invocable_v<TCallback, EBlendFunction>
+void ForEachBlendFunction(TCallback&& InCallback)
+{
+	ForEachBlendFunctionBreakable([&InCallback](EBlendFunction BlendFunction)
+	{
+		InCallback(BlendFunction); return EBreakBehavior::Continue;
+	});
+}
+template<typename TCallback> requires std::is_invocable_r_v<EBreakBehavior, TCallback, EBlendFunction>
+void ForEachBlendFunctionBreakable(TCallback&& InCallback)
+{
+	for (int32 Index = 0; Index < static_cast<int32>(EBlendFunction::Num); ++Index)
+	{
+		if (InCallback(static_cast<EBlendFunction>(Index)) == EBreakBehavior::Break)
+		{
+			break;
+		}
+	}
+}
+}
\ No newline at end of file
diff --git a/Engine/Plugins/Animation/TweeningUtils/Source/TweeningUtilsEditor/Public/Math/KeyBlendingAbstraction.h b/Engine/Plugins/Animation/TweeningUtils/Source/TweeningUtilsEditor/Public/Math/KeyBlendingAbstraction.h
deleted file mode 100644
diff --git a/Engine/Plugins/Animation/TweeningUtils/Source/TweeningUtilsEditor/Public/TweeningUtilsCommands.h b/Engine/Plugins/Animation/TweeningUtils/Source/TweeningUtilsEditor/Public/TweeningUtilsCommands.h
new file mode 100644
--- /dev/null
+++ b/Engine/Plugins/Animation/TweeningUtils/Source/TweeningUtilsEditor/Public/TweeningUtilsCommands.h
@@ -0,0 +1,28 @@
+﻿// Copyright Epic Games, Inc. All Rights Reserved.
+
+#pragma once
+
+#include "Framework/Commands/Commands.h"
+
+namespace UE::TweeningUtilsEditor
+{
+class TWEENINGUTILSEDITOR_API FTweeningUtilsCommands : public TCommands<FTweeningUtilsCommands>
+{
+public:
+	
+	TSharedPtr<FUICommandInfo> SetControlsToTween;
+	TSharedPtr<FUICommandInfo> SetTweenPushPull;
+	TSharedPtr<FUICommandInfo> SetTweenBlendNeighbor;
+	TSharedPtr<FUICommandInfo> SetTweenBlendRelative;
+	TSharedPtr<FUICommandInfo> SetTweenBlendEase;
+	TSharedPtr<FUICommandInfo> SetTweenSmoothRough;
+	TSharedPtr<FUICommandInfo> ToggleOvershootMode;
+
+	FTweeningUtilsCommands();
+
+	//~ Begin TCommands Interface
+	virtual void RegisterCommands() override;
+	//~ End TCommands Interface
+};
+}
+
diff --git a/Engine/Plugins/Animation/TweeningUtils/Source/TweeningUtilsEditor/Public/TweeningUtilsStyle.h b/Engine/Plugins/Animation/TweeningUtils/Source/TweeningUtilsEditor/Public/TweeningUtilsStyle.h
new file mode 100644
--- /dev/null
+++ b/Engine/Plugins/Animation/TweeningUtils/Source/TweeningUtilsEditor/Public/TweeningUtilsStyle.h
@@ -0,0 +1,20 @@
+﻿// Copyright Epic Games, Inc. All Rights Reserved.
+
+#pragma once
+
+#include "Styling/SlateStyle.h"
+
+namespace UE::TweeningUtilsEditor
+{
+/** Style for the widgets and commands in this module. */
+class FTweeningUtilsStyle final
+: public FSlateStyleSet
+{
+public:
+	
+	TWEENINGUTILSEDITOR_API static FTweeningUtilsStyle& Get();
+	
+	FTweeningUtilsStyle();
+	virtual ~FTweeningUtilsStyle() override;
+};
+}
diff --git a/Engine/Plugins/Animation/TweeningUtils/Source/TweeningUtilsEditor/Public/Widgets/ETweenScaleMode.h b/Engine/Plugins/Animation/TweeningUtils/Source/TweeningUtilsEditor/Public/Widgets/ETweenScaleMode.h
new file mode 100644
--- /dev/null
+++ b/Engine/Plugins/Animation/TweeningUtils/Source/TweeningUtilsEditor/Public/Widgets/ETweenScaleMode.h
@@ -0,0 +1,17 @@
+﻿// Copyright Epic Games, Inc. All Rights Reserved.
+
+#pragma once
+
+#include "HAL/Platform.h"
+
+namespace UE::TweeningUtilsEditor
+{
+/** Affects how blend values, which are normalized to the range -1.0 to 1.0, are interpreted. */
+enum class ETweenScaleMode : uint8
+{
+	/** -1.0 to 1.0 maps to -100% to +100% */
+	Normalized,
+	/** -1.0 to 1.0 maps to -200% to +200% */
+	Overshoot
+};
+}
\ No newline at end of file
diff --git a/Engine/Plugins/Animation/TweeningUtils/Source/TweeningUtilsEditor/Public/Widgets/MVC/Models/CurveTweenModel.h b/Engine/Plugins/Animation/TweeningUtils/Source/TweeningUtilsEditor/Public/Widgets/MVC/Models/CurveTweenModel.h
new file mode 100644
--- /dev/null
+++ b/Engine/Plugins/Animation/TweeningUtils/Source/TweeningUtilsEditor/Public/Widgets/MVC/Models/CurveTweenModel.h
@@ -0,0 +1,35 @@
+﻿// Copyright Epic Games, Inc. All Rights Reserved.
+
+#pragma once
+
+#include "FunctionBasedTweenModel.h"
+#include "Math/ContiguousKeyMapping.h"
+#include "Templates/SharedPointer.h"
+
+class FCurveEditor;
+
+namespace UE::TweeningUtilsEditor
+{
+/** Tween the curves of the curve editor's selection. */
+class TWEENINGUTILSEDITOR_API FCurveTweenModel : public FFunctionBasedTweenModel
+{
+public:
+
+	explicit FCurveTweenModel(TWeakPtr<FCurveEditor> InWeakCurveEditor, const TSharedRef<ITweenFunction>& InTweenFunction);
+	
+	//~ Begin FTweenModel Interface
+	virtual void StartBlendOperation() override;
+	virtual void StopBlendOperation() override;
+	virtual void BlendValues(float InNormalizedValue) override;
+	//~ Begin FTweenModel Interface
+
+private:
+
+	/** The curve editor on which to tween the curves. */
+	const TWeakPtr<FCurveEditor> WeakCurveEditor;
+	
+	/** Created in StartBlendOperation and used for the entirety of the blend operation. */
+	FContiguousKeyMapping ContiguousKeySelection;
+};
+}
+
diff --git a/Engine/Plugins/Animation/TweeningUtils/Source/TweeningUtilsEditor/Public/Widgets/MVC/Models/FunctionBasedTweenModel.h b/Engine/Plugins/Animation/TweeningUtils/Source/TweeningUtilsEditor/Public/Widgets/MVC/Models/FunctionBasedTweenModel.h
new file mode 100644
--- /dev/null
+++ b/Engine/Plugins/Animation/TweeningUtils/Source/TweeningUtilsEditor/Public/Widgets/MVC/Models/FunctionBasedTweenModel.h
@@ -0,0 +1,34 @@
+﻿// Copyright Epic Games, Inc. All Rights Reserved.
+
+#pragma once
+
+#include "Math/ContiguousKeyMapping.h"
+#include "Templates/SharedPointer.h"
+#include "TweenModel.h"
+
+namespace UE::TweeningUtilsEditor
+{
+	class ITweenFunction;
+}
+
+class FCurveEditor;
+
+namespace UE::TweeningUtilsEditor
+{
+/** Base class for models that want to implement BlendValues based off of some injectable ITweenFunction. */
+class FFunctionBasedTweenModel : public FTweenModel
+{
+public:
+
+	explicit FFunctionBasedTweenModel(const TSharedRef<ITweenFunction>& InTweenFunction) : TweenFunction(InTweenFunction) {}
+
+	void SetTweenFunction(const TSharedRef<ITweenFunction>& InTweenFunction) { TweenFunction = InTweenFunction; }
+	const TSharedRef<ITweenFunction>& GetTweenFunction() const { return TweenFunction; }
+	
+private:
+
+	/** The tween function the sub-class will use for blending. */
+	TSharedRef<ITweenFunction> TweenFunction;
+};
+}
+
diff --git a/Engine/Plugins/Animation/TweeningUtils/Source/TweeningUtilsEditor/Public/Widgets/MVC/Models/TransactionalTweenModelProxy.h b/Engine/Plugins/Animation/TweeningUtils/Source/TweeningUtilsEditor/Public/Widgets/MVC/Models/TransactionalTweenModelProxy.h
new file mode 100644
--- /dev/null
+++ b/Engine/Plugins/Animation/TweeningUtils/Source/TweeningUtilsEditor/Public/Widgets/MVC/Models/TransactionalTweenModelProxy.h
@@ -0,0 +1,46 @@
+﻿// Copyright Epic Games, Inc. All Rights Reserved.
+
+#pragma once
+
+#include "TweenModel.h"
+
+#include "Internationalization/Internationalization.h"
+#include "Internationalization/Text.h"
+#include "ScopedTransaction.h"
+#include "Templates/UniquePtr.h"
+
+#include <type_traits>
+
+class FCurveEditor;
+
+namespace UE::TweeningUtilsEditor
+{
+/** Starts a transaction in StartBlendOperation, stops it in StopBlendOperation.  */
+template<typename TBase> requires std::is_base_of_v<FTweenModel, TBase>
+class TTransactionalTweenModelProxy : public TBase
+{
+public:
+
+	template<typename... TArg>
+	explicit TTransactionalTweenModelProxy(TArg&&... Arg) : TBase(Forward<TArg>(Arg)...) {}
+	
+	//~ Begin FTweenModel Interface
+	virtual void StartBlendOperation() override
+	{
+		InProgressTransaction = MakeUnique<FScopedTransaction>(NSLOCTEXT("FTransactionalTweenModelProxy", "Transaction", "Blend values"));
+		TBase::StartBlendOperation();
+	}
+	virtual void StopBlendOperation() override
+	{
+		TBase::StopBlendOperation();
+		InProgressTransaction.Reset();
+	}
+	//~ Begin FTweenModel Interface
+
+private:
+
+	/** Active during a blend operation. */
+	TUniquePtr<FScopedTransaction> InProgressTransaction;
+};
+}
+
diff --git a/Engine/Plugins/Animation/TweeningUtils/Source/TweeningUtilsEditor/Public/Widgets/MVC/Models/TweenModel.h b/Engine/Plugins/Animation/TweeningUtils/Source/TweeningUtilsEditor/Public/Widgets/MVC/Models/TweenModel.h
new file mode 100644
--- /dev/null
+++ b/Engine/Plugins/Animation/TweeningUtils/Source/TweeningUtilsEditor/Public/Widgets/MVC/Models/TweenModel.h
@@ -0,0 +1,67 @@
+﻿// Copyright Epic Games, Inc. All Rights Reserved.
+
+#pragma once
+
+#include "Misc/Optional.h"
+#include "Widgets/ETweenScaleMode.h"
+
+namespace UE::TweeningUtilsEditor
+{
+/**
+ * Implements the logic of tweening some values: it could be curves, it could be control rig models.
+ * Acts as model in a Model-View-Controller architecture.
+ */
+class TWEENINGUTILSEDITOR_API FTweenModel
+{
+public:
+
+	/** Called when a blend operation is started. */
+	virtual void StartBlendOperation() {}
+	/** Called when a blend operation is stopped. */
+	virtual void StopBlendOperation() {}
+	/**
+	 * Does the blending based on the InNormalizedValue.
+	 * @param InNormalizedValue Value in [-1, 1]
+	 */
+	virtual void BlendValues(float InNormalizedValue) = 0;
+
+	/**
+	 * Blends to a single value once.
+	 * @param InNormalizedValue Value in [-1, 1]
+	 */
+	void BlendOneOff(float InNormalizedValue)
+	{
+		StartBlendOperation();
+		BlendValues(InNormalizedValue);
+		StopBlendOperation();
+	}
+
+	virtual void SetScaleMode(ETweenScaleMode InMode) { ScaleMode = InMode; }
+	ETweenScaleMode GetScaleMode() const { return ScaleMode; }
+	
+	/**
+	 * @param InNormalizedValue Value from -1 to 1.
+	 * @return Value scaled according to the ScaleMode setting.
+	 */
+	float ScaleBlendValue(float InNormalizedValue) const
+	{
+		return ScaleMode == ETweenScaleMode::Normalized
+			? InNormalizedValue
+			: 2.f * InNormalizedValue;
+	}
+	
+	virtual ~FTweenModel() = default;
+	
+private:
+
+	/**
+	 * Affects how normalized values sent to BlendValueNormalizedAndRecomputeKeys are interpreted:
+	 * - False: -1.0 to 1.0 maps to -100% to +100%
+	 * - True: -1.0 to 1.0 maps to -200% to +200%
+	 */
+	ETweenScaleMode ScaleMode = ETweenScaleMode::Normalized;
+	
+	TOptional<float> CurrentBlendValue;
+};
+}
+
diff --git a/Engine/Plugins/Animation/TweeningUtils/Source/TweeningUtilsEditor/Public/Widgets/MVC/STweenView.h b/Engine/Plugins/Animation/TweeningUtils/Source/TweeningUtilsEditor/Public/Widgets/MVC/STweenView.h
new file mode 100644
--- /dev/null
+++ b/Engine/Plugins/Animation/TweeningUtils/Source/TweeningUtilsEditor/Public/Widgets/MVC/STweenView.h
@@ -0,0 +1,43 @@
+﻿// Copyright Epic Games, Inc. All Rights Reserved.
+
+#pragma once
+
+#include "Widgets/DeclarativeSyntaxSupport.h"
+#include "Widgets/SCompoundWidget.h"
+
+namespace UE::TweeningUtilsEditor
+{
+class FTweenModel;
+enum class ETweenScaleMode : uint8;
+	
+/**
+ * It bridges the STweenSlider, which acts as the view, and the model, FTweenModel, by interchanging data between them.
+ * Acts as view in a Model-View-Controller architecture.
+ */
+class TWEENINGUTILSEDITOR_API STweenView : public SCompoundWidget
+{
+public:
+
+	SLATE_BEGIN_ARGS(STweenView) {}
+		/** The icon to place in the slider button. */
+		SLATE_ATTRIBUTE(const FSlateBrush*, SliderIcon)
+		/** The main color. It tints the slider button and the points. */
+		SLATE_ATTRIBUTE(FLinearColor, SliderColor)
+	SLATE_END_ARGS()
+
+	void Construct(const FArguments& InArgs, const TSharedRef<FTweenModel>& InTweenModel);
+
+private:
+
+	/** The model that these controls are driving */
+	TSharedPtr<FTweenModel> TweenModel;
+	
+	ETweenScaleMode GetBarRenderMode() const;
+	
+	void OnDragStarted();
+	void OnDragEnded();
+	void OnDragValueUpdated(float Value);
+	void OnPointPicked(float Value);
+	float MapSliderValueToBlendValue(float Value);
+};
+}
diff --git a/Engine/Plugins/Animation/TweeningUtils/Source/TweeningUtilsEditor/Public/Widgets/MVC/TweenToolbarController.h b/Engine/Plugins/Animation/TweeningUtils/Source/TweeningUtilsEditor/Public/Widgets/MVC/TweenToolbarController.h
new file mode 100644
--- /dev/null
+++ b/Engine/Plugins/Animation/TweeningUtils/Source/TweeningUtilsEditor/Public/Widgets/MVC/TweenToolbarController.h
@@ -0,0 +1,74 @@
+﻿// Copyright Epic Games, Inc. All Rights Reserved.
+
+#pragma once
+
+#include "Templates/SharedPointer.h"
+#include "Templates/UnrealTemplate.h"
+
+class FText;
+class FToolBarBuilder;
+class FUICommandList;
+class SWidget;
+struct FLinearColor;
+struct FSlateBrush;
+struct FSlateIcon;
+
+namespace UE::TweeningUtilsEditor
+{
+class FFunctionBasedTweenModel;
+class ITweenFunction;
+class ITweenFunctionContainer;
+	
+/**
+ * Binds together a STweenView and FTweenModel in the context of adding the controls to a toolbar widget.
+ * Acts as controller in a Model-View-Controller architecture.
+ *
+ * Makes a toolbar as follows |Combo|Slider|Overshot|
+ * - Combo: Selects the blend function
+ * - Slider: Adjusts the blend value and updates the FTweenModel
+ * - Overshot: Toggle button for changing the overshot property on FTweenModel
+ */
+class TWEENINGUTILSEDITOR_API FTweenToolbarController : public FNoncopyable
+{
+public:
+	
+	/**
+	 * @param InModel The model that the UI will be driving.
+	 * @param InCommandList Used to bind the commands added to the toolbar.
+	 * @param InTweenFunctions The functions that can be selected
+	 */
+	explicit FTweenToolbarController(
+		const TSharedRef<FFunctionBasedTweenModel>& InModel,
+		const TSharedRef<FUICommandList>& InCommandList,
+		const TSharedRef<ITweenFunctionContainer>& InTweenFunctions
+		);
+	~FTweenToolbarController();
+	
+	/** Extends the toolbar. Safe to call multiple times (although typically you'll call this exactly once). */
+	void AddToToolbar(FToolBarBuilder& ToolbarBuilder) const;
+
+private:
+
+	/** Receives the blend values */
+	const TSharedRef<FFunctionBasedTweenModel> TweenModel;
+	/** Used when the combo button is opened. */
+	const TSharedRef<FUICommandList> CommandList;
+	/** Holds the functions that can be selected and gives us information about how to display them. */
+	const TSharedRef<ITweenFunctionContainer> TweenFunctions;
+
+	/** Creates the drop-down list for selecting the tween mode. */
+	TSharedRef<SWidget> MakeTweenModeMenu() const;
+	FText GetLabelForComboBox() const;
+	FText GetToolTipForComboBox() const;
+	FSlateIcon GetIconForComboBox() const;
+	
+	void SetBlendFunction(TSharedPtr<ITweenFunction> InFunction) const;
+	bool IsBlendFunctionSelected(TSharedPtr<ITweenFunction> InFunction) const;
+	
+	void ToggleOvershootMode() const;
+	bool IsOvershootModeEnabled() const;
+	
+	const FSlateBrush* GetIconForSlider() const;
+	FLinearColor GetColorForSlider() const;
+};
+}
diff --git a/Engine/Plugins/Animation/TweeningUtils/Source/TweeningUtilsEditor/Public/Widgets/STweenSlider.h b/Engine/Plugins/Animation/TweeningUtils/Source/TweeningUtilsEditor/Public/Widgets/STweenSlider.h
new file mode 100644
--- /dev/null
+++ b/Engine/Plugins/Animation/TweeningUtils/Source/TweeningUtilsEditor/Public/Widgets/STweenSlider.h
@@ -0,0 +1,117 @@
+﻿// Copyright Epic Games, Inc. All Rights Reserved.
+
+#pragma once
+
+#include "ETweenScaleMode.h"
+#include "TweeningUtilsStyle.h"
+#include "TweenSliderStyle.h"
+#include "Framework/DelayedDrag.h"
+#include "Widgets/DeclarativeSyntaxSupport.h"
+#include "Widgets/SLeafWidget.h"
+
+namespace UE::TweeningUtilsEditor
+{
+struct FSliderWidgetData;
+	
+/**
+ * The slider ranges from -1 to 1. It drives the blend while pressed down. When released, it jumps to 0: letting the user reapply the blend again.
+ 
+ * The slider looks like this:
+ *  - ETweenScaleMode::Normalized: [   +   +   +   [I]   +   +   +   ]
+ *  - ETweenScaleMode::Overshoot:  [ + + + | + + + [I] + + + | + + + ]
+ * Here:
+ *	- I is the Icon and the bit behind it is the SliderIconBackground. Together, they form a button.
+ *	- + is a point on the scale.
+ */
+class TWEENINGUTILSEDITOR_API STweenSlider : public SLeafWidget
+{
+public:
+
+	DECLARE_DELEGATE_OneParam(FSliderChangedDelegate, float);
+	DECLARE_DELEGATE_RetVal_OneParam(float, FMapSliderValueToBlendValue, float);
+
+	SLATE_BEGIN_ARGS(STweenSlider)
+		: _Style(&FTweeningUtilsStyle::Get().GetWidgetStyle<FTweenSliderStyle>("TweenSlider"))
+		, _ColorAndOpacity(FLinearColor::White)
+		, _SliderIcon(nullptr)
+		, _SliderColor(FLinearColor(254.f/255.f, 155.f/255.f, 7.f/255.f))
+		, _ScaleRenderMode(ETweenScaleMode::Normalized)
+		{}
+		
+		/** The visual style of the button */
+		SLATE_STYLE_ARGUMENT( FTweenSliderStyle, Style )
+
+		/** The root opacity of the widget. */
+		SLATE_ATTRIBUTE(FSlateColor, ColorAndOpacity)
+		
+		/** The icon to place in the slider button. */
+		SLATE_ATTRIBUTE(const FSlateBrush*, SliderIcon)
+		/** The main color. It tints the slider button and the points. */
+		SLATE_ATTRIBUTE(FLinearColor, SliderColor)
+		/** Affects how the scale is rendered. Events keep returning the -1 to 1 mode - you should compute the scale yourself.*/
+		SLATE_ATTRIBUTE(ETweenScaleMode, ScaleRenderMode)
+		
+		/** Invoked when the slider begins being dragged. */
+		SLATE_EVENT(FSimpleDelegate, OnSliderDragStarted)
+		/** Invoked when the slider stops being dragged */
+		SLATE_EVENT(FSimpleDelegate, OnSliderDragEnded)
+		/** Invoked for as long as the slider is being dragged. Receives a value from -1 to 1 regardless of ScaleRenderMode. */
+		SLATE_EVENT(FSliderChangedDelegate, OnSliderValueDragged)
+		/** Invoked when the user presses a point to select a blend value. Receives a value from -1 to 1 regardless of ScaleRenderMode. */
+		SLATE_EVENT(FSliderChangedDelegate, OnPointValuePicked)
+
+		/** Invoked to display the correct blend value tooltips. Receives a [-1,1] value and outputs the blend value. */
+		SLATE_EVENT(FMapSliderValueToBlendValue, MapSliderValueToBlendValue)
+		
+	SLATE_END_ARGS()
+
+	void Construct(const FArguments& InArgs);
+
+private:
+
+	/**
+	 * Abstracts the implementation.
+	 *
+	 * Specifically, it hides FTweenSliderHoverState, FTweenSliderDrawArgs, and FTweenWidgetArgs so TweenSliderDrawUtils.h does not become part
+	 * of the public API.
+	 */
+	TPimplPtr<FSliderWidgetData> Pimpl;
+	
+	/** Invoked when the slider begins being dragged. */
+	FSimpleDelegate OnSliderDragStartedDelegate;
+	/** Invoked when the slider stops being dragged */
+	FSimpleDelegate OnSliderDragStoppedDelegate;
+	/** Invoked during drag when the slider value changes. */
+	FSliderChangedDelegate OnSliderChangedDelegate;
+	/** Invoked when the user presses a point to select a blend value. */
+	FSliderChangedDelegate OnPointValuePickedDelegate;
+
+	/** Invoked to display the correct blend value tooltips. */
+	FMapSliderValueToBlendValue MapSliderValueToBlendValueDelegate;
+	
+	//~ Begin SLeafWidget Interface
+	virtual FReply OnMouseButtonDoubleClick(const FGeometry& InMyGeometry, const FPointerEvent& InMouseEvent) override;
+	virtual FReply OnMouseButtonDown(const FGeometry& InGeometry, const FPointerEvent& InMouseEvent) override;
+	virtual FReply OnMouseButtonUp(const FGeometry& InGeometry, const FPointerEvent& InMouseEvent) override;
+	virtual void OnMouseLeave(const FPointerEvent& InMouseEvent) override;
+	virtual FReply OnMouseMove(const FGeometry& MyGeometry, const FPointerEvent& InMouseEvent) override;
+	virtual void OnFinishedPointerInput() override;
+	virtual FVector2D ComputeDesiredSize(float) const override;
+	virtual void Tick(const FGeometry& AllottedGeometry, const double InCurrentTime, const float InDeltaTime) override;
+	virtual int32 OnPaint(const FPaintArgs& InArgs, const FGeometry& InAllottedGeometry, const FSlateRect& InCullingRect, FSlateWindowElementList& OutDrawElements, int32 InLayerId, const FWidgetStyle& InWidgetStyle, bool bParentEnabled) const override;
+	virtual TSharedPtr<IToolTip> GetToolTip() override;
+	//~ End SLeafWidget Interface
+
+	/** Updates LastMousePositionOnSliderBar. */
+	void UpdateLastMousePositionOnSlider(const FGeometry& InGeometry, const FVector2D& InMousePos);
+
+	bool IsMouseDown() const;
+	bool IsDragging() const;
+
+	/** If the user is currently hovering a point, set the blend value to that point's value.  */
+	void HandlePickValueOfCurrentlyHoveredPoint(const FGeometry& InGeometry, const FPointerEvent& InMouseEvent) const;
+
+	/** Gets the tooltip text for what is currently being hovered. */
+	FText GetToolTipText() const;
+};
+}
diff --git a/Engine/Plugins/Animation/TweeningUtils/Source/TweeningUtilsEditor/Public/Widgets/TweenSliderStyle.h b/Engine/Plugins/Animation/TweeningUtils/Source/TweeningUtilsEditor/Public/Widgets/TweenSliderStyle.h
new file mode 100644
--- /dev/null
+++ b/Engine/Plugins/Animation/TweeningUtils/Source/TweeningUtilsEditor/Public/Widgets/TweenSliderStyle.h
@@ -0,0 +1,140 @@
+﻿// Copyright Epic Games, Inc. All Rights Reserved.
+
+#pragma once
+
+#include "Layout/Margin.h"
+#include "Math/Vector2D.h"
+#include "Styling/SlateBrush.h"
+#include "Styling/SlateWidgetStyle.h"
+#include "TweenSliderStyle.generated.h"
+
+USTRUCT()
+struct TWEENINGUTILSEDITOR_API FTweenPointStyle
+{
+	GENERATED_BODY()
+	
+	FTweenPointStyle();
+	explicit FTweenPointStyle(
+		const FVector2D& InNormalSize, const FVector2D& InHoveredSize, const FVector2D& InPressedSize, const FVector2D& InHitSize
+		);
+	
+	/** Brush when a point on the bar is not hovered or pressed */
+	UPROPERTY(EditAnywhere, Category=Appearance)
+	FSlateBrush Normal;
+	FTweenPointStyle& SetNormal( const FSlateBrush& InNormal ){ Normal = InNormal; return *this; }
+
+	/** Brush when a point on the bar is hovered */
+	UPROPERTY(EditAnywhere, Category=Appearance)
+	FSlateBrush Hovered;
+	FTweenPointStyle& SetHovered( const FSlateBrush& InHovered){ Hovered = InHovered; return *this; }
+
+	/** Brush when a point on the bar is pressed */
+	UPROPERTY(EditAnywhere, Category=Appearance)
+	FSlateBrush Pressed;
+	FTweenPointStyle& Set( const FSlateBrush& InPressed ){ Pressed = InPressed; return *this; }
+
+	/** Brush when the slider has passed a point*/
+	UPROPERTY(EditAnywhere, Category=Appearance)
+	FSlateBrush PassedPoint;
+	FTweenPointStyle& SetPassed( const FSlateBrush& InPassed ){ PassedPoint = InPassed; return *this; }
+	
+	/** The size of the hit box against which hit tests are made for this point (so the user does not need to click them pixel perfectly). */
+	UPROPERTY(EditAnywhere, Category=Appearance)
+	FVector2D HitTestSize;
+	FTweenPointStyle& SetHitTestSize( const FVector2D& InSize ){ HitTestSize = InSize; return *this; }
+	
+	void GetResources(TArray<const FSlateBrush*>& OutBrushes) const;
+};
+
+/** Style for STweenSlider */
+USTRUCT()
+struct TWEENINGUTILSEDITOR_API FTweenSliderStyle : public FSlateWidgetStyle
+{
+	GENERATED_BODY()
+
+	FTweenSliderStyle();
+
+	//~ Begin FSlateWidgetStyle Interface
+	virtual void GetResources(TArray<const FSlateBrush*>& OutBrushes) const override;
+	//~ End FSlateWidgetStyle Interface
+	
+	static const FName TypeName;
+	virtual const FName GetTypeName() const override { return TypeName; }
+	
+	static const FTweenSliderStyle& GetDefault();
+	
+	/** The dimensions of the slider. */
+	UPROPERTY(EditAnywhere, Category=Appearance)
+	FVector2D BarDimensions;
+	FTweenSliderStyle& SetBarDimensions( const FVector2D& InDimensions ){ BarDimensions = InDimensions; return *this; }
+
+	/** Brush of the bar on which the points are drawn. The slider brush is drawn over it. Usually the bar's brush is smaller than the button. */
+	UPROPERTY(EditAnywhere, Category=Appearance)
+	FSlateBrush BarBrush;
+	FTweenSliderStyle& SetBarBrush( const FSlateBrush& InBarBrush ){ BarBrush = InBarBrush; return *this; }
+
+	
+	
+	/** Brush when the button is not hovered or pressed */
+	UPROPERTY(EditAnywhere, Category=Appearance)
+	FSlateBrush NormalSliderButton;
+	FTweenSliderStyle& SetNormalSliderButton( const FSlateBrush& InNormalBrush ){ NormalSliderButton = InNormalBrush; return *this; }
+
+	/** Brush when hovered */
+	UPROPERTY(EditAnywhere, Category=Appearance)
+	FSlateBrush HoveredSliderButton;
+	FTweenSliderStyle& SetHoveredSliderButton( const FSlateBrush& InHoveredBrush){ HoveredSliderButton = InHoveredBrush; return *this; }
+
+	/** Brush when pressed */
+	UPROPERTY(EditAnywhere, Category=Appearance)
+	FSlateBrush PressedSliderButton;
+	FTweenSliderStyle& SetPressedSliderButton( const FSlateBrush& InPressedBrush ){ PressedSliderButton = InPressedBrush; return *this; }
+
+
+	
+	/** Tint for the icon when the button is not hovered or pressed */
+	UPROPERTY(EditAnywhere, Category=Appearance)
+	FSlateColor NormalIconTint;
+	FTweenSliderStyle& SetNormalIconTint( const FSlateColor& InNormalIconTint ){ NormalIconTint = InNormalIconTint; return *this; }
+
+	/** Tint for the icon when hovered */
+	UPROPERTY(EditAnywhere, Category=Appearance)
+	FSlateColor HoveredIconTint;
+	FTweenSliderStyle& SetHoveredIconTint( const FSlateColor& InHoveredIconTint){ HoveredIconTint = InHoveredIconTint; return *this; }
+
+	/** Tint for the icon when pressed */
+	UPROPERTY(EditAnywhere, Category=Appearance)
+	FSlateColor PressedIconTint;
+	FTweenSliderStyle& SetPressedIconTint( const FSlateColor& InPressedIconTint ){ PressedIconTint = InPressedIconTint; return *this; }
+
+
+	
+	/** Brushes for the small points on the bar. */
+	UPROPERTY(EditAnywhere, Category=Appearance)
+	FTweenPointStyle SmallPoint;
+	FTweenSliderStyle& SetSmallPoint( const FTweenPointStyle& InNormal ){ SmallPoint = InNormal; return *this; }
+
+	/** When overshoot mode is enabled, the bars at 100% and -100%. */
+	UPROPERTY(EditAnywhere, Category=Appearance)
+	FTweenPointStyle MediumPoint;
+	FTweenSliderStyle& SetMediumPoint( const FTweenPointStyle& InHovered){ MediumPoint = InHovered; return *this; }
+
+	/** Brushes for points on the left or right end. */
+	UPROPERTY(EditAnywhere, Category=Appearance)
+	FTweenPointStyle EndPoint;
+	FTweenSliderStyle& SetEndPoint( const FTweenPointStyle& InPressed ){ EndPoint = InPressed; return *this; }
+	
+
+	
+	/** Brush drawn from center to slider when moving the slider (should be a little transparent, drawn over points but under the slider). */
+	UPROPERTY(EditAnywhere, Category=Appearance)
+	FSlateBrush PassedValueBackground;
+	FTweenSliderStyle& SetPassedValueBackground( const FSlateBrush& InBrush ){ PassedValueBackground = InBrush; return *this; }
+
+
+	
+	/** Padding of the icon placed in the slider button */
+	UPROPERTY(EditAnywhere, Category=Appearance)
+	FMargin IconPadding;
+	FTweenSliderStyle& SetIconPadding( const FMargin& InIconPadding ){ IconPadding = InIconPadding; return *this; }
+};
\ No newline at end of file
diff --git a/Engine/Plugins/Animation/TweeningUtils/Source/TweeningUtilsEditor/TweeningUtilsEditor.build.cs b/Engine/Plugins/Animation/TweeningUtils/Source/TweeningUtilsEditor/TweeningUtilsEditor.build.cs
--- a/Engine/Plugins/Animation/TweeningUtils/Source/TweeningUtilsEditor/TweeningUtilsEditor.build.cs
+++ b/Engine/Plugins/Animation/TweeningUtils/Source/TweeningUtilsEditor/TweeningUtilsEditor.build.cs
@@ -11,10 +11,13 @@ public class TweeningUtilsEditor : ModuleRules
 		PublicDependencyModuleNames.AddRange(new string[]
 		{
 			"Core", 
+			"CoreUObject",
 			"CurveEditor",
+			"Engine",
 			"InputCore",
 			"SlateCore",
-			"Slate"
+			"Slate",
+			"UnrealEd"
 		});
 		
 		PrivateDependencyModuleNames.AddRange(new string[]
