Fix static analysis warnings

diff --git a/Engine/Plugins/Animation/TweeningUtils/Source/TweeningUtilsEditor/Private/Widgets/TweenSliderDrawUtils.cpp b/Engine/Plugins/Animation/TweeningUtils/Source/TweeningUtilsEditor/Private/Widgets/TweenSliderDrawUtils.cpp
--- a/Engine/Plugins/Animation/TweeningUtils/Source/TweeningUtilsEditor/Private/Widgets/TweenSliderDrawUtils.cpp
+++ b/Engine/Plugins/Animation/TweeningUtils/Source/TweeningUtilsEditor/Private/Widgets/TweenSliderDrawUtils.cpp
@@ -360,10 +360,9 @@ static void DrawButton(
 
 	const FSlateColor& BaseIconTint = InDrawArgs.bDrawButtonPressed
 		? Style->PressedIconTint
-		: InDrawArgs.HoverState.bIsSliderHovered ? Style->PressedIconTint : Style->PressedIconTint;
+		: InDrawArgs.HoverState.bIsSliderHovered ? Style->HoveredIconTint : Style->PressedIconTint;
 	
 	const FLinearColor ColorAndOpacity = InWidgetArgs.ColorAndOpacity.Get().GetColor(InWidgetStyle);
-	const FLinearColor IconTint = ColorAndOpacity * BaseIconTint.GetColor(InWidgetStyle) * SliderIconBrush->TintColor.GetColor(InWidgetStyle);
 	const FLinearColor ButtonTint = ColorAndOpacity * InWidgetArgs.SliderColor.Get() * SliderButtonBrush->TintColor.GetColor(InWidgetStyle);
 	
 	FSlateDrawElement::MakeBox(
@@ -371,6 +370,7 @@ static void DrawButton(
 		);
 	if (SliderIconBrush)
 	{
+		const FLinearColor IconTint = ColorAndOpacity * BaseIconTint.GetColor(InWidgetStyle) * SliderIconBrush->TintColor.GetColor(InWidgetStyle);
 		FSlateDrawElement::MakeBox(
 			OutDrawElements, LayerId + 1, InDrawArgs.IconArea.ToPaintGeometry(), SliderIconBrush, ESlateDrawEffect::None, IconTint
 		);
