Curve editor: add indirect mouse sliding for tweening
- The user can now press U + move mouse to slide the tween slider in the durve editor's toolbar
- In the future, ControlRig will be made to use this logic instead (right now it is somewhat duplicated)

diff --git a/Engine/Plugins/Animation/TweeningUtils/Source/TweeningUtilsEditor/Private/Math/Abstraction/ITweenFunctionContainer.cpp b/Engine/Plugins/Animation/TweeningUtils/Source/TweeningUtilsEditor/Private/Math/Abstraction/ITweenFunctionContainer.cpp
new file mode 100644
--- /dev/null
+++ b/Engine/Plugins/Animation/TweeningUtils/Source/TweeningUtilsEditor/Private/Math/Abstraction/ITweenFunctionContainer.cpp
@@ -0,0 +1,22 @@
+﻿// Copyright Epic Games, Inc. All Rights Reserved.
+
+#include "Math/Abstraction/ITweenFunctionContainer.h"
+
+#include "Math/Abstraction/ITweenFunction.h"
+#include "Misc/CoreMiscDefines.h"
+#include "Templates/SharedPointer.h"
+
+namespace UE::TweeningUtilsEditor
+{
+int32 ITweenFunctionContainer::IndexOf(const TSharedRef<ITweenFunction>& InFunction) const
+{
+	for (int32 Index = 0; Index < NumFunctions(); ++Index)
+	{
+		if (GetFunction(Index) == InFunction)
+		{
+			return Index;
+		}
+	}
+	return INDEX_NONE;
+}
+}
diff --git a/Engine/Plugins/Animation/TweeningUtils/Source/TweeningUtilsEditor/Private/TweeningUtilsCommands.cpp b/Engine/Plugins/Animation/TweeningUtils/Source/TweeningUtilsEditor/Private/TweeningUtilsCommands.cpp
--- a/Engine/Plugins/Animation/TweeningUtils/Source/TweeningUtilsEditor/Private/TweeningUtilsCommands.cpp
+++ b/Engine/Plugins/Animation/TweeningUtils/Source/TweeningUtilsEditor/Private/TweeningUtilsCommands.cpp
@@ -29,7 +29,10 @@ void FTweeningUtilsCommands::RegisterCommands()
 	UI_COMMAND(SetTweenBlendRelative, "Set Blend Relative", "Sets the tween mode to Blend Relative", EUserInterfaceActionType::RadioButton, FInputChord() );
 	UI_COMMAND(SetTweenBlendEase, "Set Blend Ease", "Sets the tween mode to Blend Ease", EUserInterfaceActionType::RadioButton, FInputChord() );
 	UI_COMMAND(SetTweenSmoothRough, "Set Smooth / Rough", "Sets the tween mode to Smooth / Rough", EUserInterfaceActionType::RadioButton, FInputChord() );
-	UI_COMMAND(ToggleOvershootMode, "Overshoot mode", "Overshoot extends the available range over 100%", EUserInterfaceActionType::ToggleButton, FInputChord() ); 
+	UI_COMMAND(ToggleOvershootMode, "Overshoot mode", "Overshoot extends the available range over 100%", EUserInterfaceActionType::ToggleButton, FInputChord() );
+	
+	UI_COMMAND(DragAnimSliderTool, "Drag Anim Slider Tool", "Drag the anim slider", EUserInterfaceActionType::Button, FInputChord(EKeys::U));
+	UI_COMMAND(ChangeAnimSliderTool, "Change Anim Slider Tool", "Go to the next blend function", EUserInterfaceActionType::Button, FInputChord(EKeys::U, EModifierKey::Shift));
 }
 }
 
diff --git a/Engine/Plugins/Animation/TweeningUtils/Source/TweeningUtilsEditor/Private/Widgets/MVC/CycleFunctionController.cpp b/Engine/Plugins/Animation/TweeningUtils/Source/TweeningUtilsEditor/Private/Widgets/MVC/CycleFunctionController.cpp
new file mode 100644
--- /dev/null
+++ b/Engine/Plugins/Animation/TweeningUtils/Source/TweeningUtilsEditor/Private/Widgets/MVC/CycleFunctionController.cpp
@@ -0,0 +1,50 @@
+﻿// Copyright Epic Games, Inc. All Rights Reserved.
+
+#include "Widgets/MVC/CycleFunctionController.h"
+
+#include "Framework/Commands/UICommandInfo.h"
+#include "Framework/Commands/UICommandList.h"
+#include "Math/Abstraction/ITweenFunctionContainer.h"
+#include "TweeningUtilsCommands.h"
+#include "Widgets/MVC/Models/FunctionBasedTweenModel.h"
+
+namespace UE::TweeningUtilsEditor
+{
+FCycleFunctionController::FCycleFunctionController(
+	const TSharedRef<FFunctionBasedTweenModel>& InTweenModel,
+	const TSharedRef<ITweenFunctionContainer>& InTweenFunctions,
+	const TSharedRef<FUICommandList>& InCommandList,
+	TSharedPtr<FUICommandInfo> InCycleCommand
+	)
+	: TweenModel(InTweenModel)
+	, TweenFunctions(InTweenFunctions)
+	, CommandList(InCommandList)
+	, CycleCommand(MoveTemp(InCycleCommand))
+{
+	check(CycleCommand);
+	CommandList->MapAction(CycleCommand, FExecuteAction::CreateRaw(this, &FCycleFunctionController::CycleToNextFunction));
+}
+
+FCycleFunctionController::FCycleFunctionController(
+	const TSharedRef<FFunctionBasedTweenModel>& InTweenModel,
+	const TSharedRef<ITweenFunctionContainer>& InTweenFunctions,
+	const TSharedRef<FUICommandList>& InCommandList
+	)
+	: FCycleFunctionController(InTweenModel, InTweenFunctions, InCommandList, FTweeningUtilsCommands::Get().ChangeAnimSliderTool)
+{}
+
+FCycleFunctionController::~FCycleFunctionController()
+{
+	CommandList->UnmapAction(CycleCommand);
+}
+
+void FCycleFunctionController::CycleToNextFunction()
+{
+	const int32 Index = TweenFunctions->IndexOf(TweenModel->GetTweenFunction());
+	if (ensure(TweenFunctions->IsValidIndex(Index)))
+	{
+		const int32 NextIndex = TweenFunctions->IsValidIndex(Index + 1) ? Index + 1 : 0;
+		TweenModel->SetTweenFunction(TweenFunctions->GetFunction(NextIndex).ToSharedRef());
+	}
+}
+}
diff --git a/Engine/Plugins/Animation/TweeningUtils/Source/TweeningUtilsEditor/Private/Widgets/MVC/MouseSlidingController.cpp b/Engine/Plugins/Animation/TweeningUtils/Source/TweeningUtilsEditor/Private/Widgets/MVC/MouseSlidingController.cpp
new file mode 100644
--- /dev/null
+++ b/Engine/Plugins/Animation/TweeningUtils/Source/TweeningUtilsEditor/Private/Widgets/MVC/MouseSlidingController.cpp
@@ -0,0 +1,223 @@
+﻿// Copyright Epic Games, Inc. All Rights Reserved.
+
+#include "Widgets/MVC/MouseSlidingController.h"
+
+#include "Framework/Application/IInputProcessor.h"
+#include "Framework/Application/SlateApplication.h"
+#include "Framework/Commands/UICommandInfo.h"
+#include "Framework/Commands/UICommandList.h"
+
+#include "HAL/PlatformMisc.h"
+#if PLATFORM_MICROSOFT
+// Not sure whether this is needed to use RECT since we're already including "HAL/PlatformMisc.h" but other code uses it, too...
+#include "Microsoft/WindowsHWrapper.h"
+#endif
+
+namespace UE::TweeningUtilsEditor
+{
+/**
+ * Orchestrates calls to StartSliding, StopSliding, and UpdateSliding
+ * The command system does not expose: 1. detecting key up, and 2. mouse movement. This class handles that.
+ * 
+ * Created once FMouseSlidingController::DragSliderCommand has been invoked and destroyed once FMouseSlidingController::DragSliderCommand's key
+ * bindings are released, which is detected by this class.
+ *
+ * FMouseSlidingController::DragSliderCommand is only invoked when the tab owning the command list, e.g. Curve Editor, is focused.
+ * Until the user releases the key that triggers the interaction, it is fine to preprocess input - we'll act as if the tab continues to be focused
+ * for the duration of the press.
+ */
+class FMouseSlidingInputProcessor : public IInputProcessor
+{
+	FMouseSlidingController& Owner;
+	bool bIsLeftMouseButtonDown = false;
+	bool bHasStopped = false;
+
+	TOptional<FVector2D> LastMouse;
+
+	bool IsDragCommandKeyUp(const FKeyEvent& InKeyEvent) const
+	{
+		bool bIsMovingSlider = false;
+		
+		for (uint32 i = 0; i < static_cast<uint8>(EMultipleKeyBindingIndex::NumChords); ++i)
+		{
+			EMultipleKeyBindingIndex ChordIndex = static_cast<EMultipleKeyBindingIndex>(i);
+			const FInputChord& Chord = *Owner.DragSliderCommand->GetActiveChord(ChordIndex);
+			bIsMovingSlider |= Chord.IsValidChord() && InKeyEvent.GetKey() == Chord.Key;
+		}
+		
+		return bIsMovingSlider;
+	}
+	
+public:
+
+	FMouseSlidingInputProcessor(FMouseSlidingController& InOwner UE_LIFETIMEBOUND) : Owner(InOwner) {}
+	
+	//~ Begin IInputProcessor Interface
+	virtual bool HandleKeyUpEvent(FSlateApplication&, const FKeyEvent& InKeyEvent) override
+	{
+		bHasStopped |= IsDragCommandKeyUp(InKeyEvent);
+		return true;
+	}
+
+	virtual bool HandleMouseButtonDownEvent(FSlateApplication& SlateApp, const FPointerEvent& InMouseEvent) override
+	{
+		if (InMouseEvent.GetEffectingButton() == EKeys::LeftMouseButton)
+		{
+			bIsLeftMouseButtonDown = true;
+			return true;
+		}
+		return false;
+	}
+	virtual bool HandleMouseButtonUpEvent(FSlateApplication& SlateApp, const FPointerEvent& InMouseEvent) override
+	{
+		if (InMouseEvent.GetEffectingButton() == EKeys::LeftMouseButton)
+		{
+			bIsLeftMouseButtonDown = false;
+			return true;
+		}
+		return false;
+	}
+
+	virtual bool HandleMouseMoveEvent(FSlateApplication&, const FPointerEvent& MouseEvent) override
+	{
+		LastMouse = MouseEvent.GetScreenSpacePosition();
+		// We handle the event; this prevents us from formally hovering any widgets for the duration of the operation.
+		return bIsLeftMouseButtonDown;
+	}
+	
+	virtual void Tick(const float, FSlateApplication&, TSharedRef<ICursor> InCursor) override
+	{
+		if (bHasStopped)
+		{
+			Owner.StopListeningForMouseEvents(*InCursor);
+			return;
+		}
+
+		if (!bIsLeftMouseButtonDown && Owner.SlidingState->IsSliding())
+		{
+			Owner.StopSliding(*InCursor);
+			return;
+		}
+		
+		if (!LastMouse || !bIsLeftMouseButtonDown)
+		{
+			return;
+		}
+
+		const FVector2D Mouse = *LastMouse;
+		LastMouse.Reset();
+		
+		if (bIsLeftMouseButtonDown && !Owner.SlidingState->IsSliding())
+		{
+			Owner.StartSliding(Mouse, *InCursor);
+		}
+		Owner.UpdateSliding(Mouse);
+	}
+	//~ Begin IInputProcessor Interface
+};
+	
+FMouseSlidingController::FMouseSlidingController(
+	TAttribute<float> InMaxSlideWidthAttr,
+	const TSharedRef<FUICommandList>& InCommandList,
+	TSharedPtr<FUICommandInfo> InDragSliderCommand
+	)
+	: MaxSlideWidthAttr(MoveTemp(InMaxSlideWidthAttr))
+	, CommandList(InCommandList)
+	, DragSliderCommand(MoveTemp(InDragSliderCommand))
+{
+	check(InMaxSlideWidthAttr.IsSet() || InMaxSlideWidthAttr.IsBound());
+	check(DragSliderCommand);
+	
+	CommandList->MapAction(
+		DragSliderCommand, FExecuteAction::CreateRaw(this, &FMouseSlidingController::StartListeningForMouseEvents)
+		);
+}
+
+FMouseSlidingController::~FMouseSlidingController()
+{
+	CommandList->UnmapAction(DragSliderCommand);
+	SlidingState.Reset();
+}
+
+FMouseSlidingController::FSlidingState::FSlidingState(FMouseSlidingController& InOwner)
+	: InputProcessor(MakeShared<FMouseSlidingInputProcessor>(InOwner))
+{
+	FSlateApplication::Get().RegisterInputPreProcessor(InputProcessor);
+}
+
+FMouseSlidingController::FSlidingState::~FSlidingState()
+{
+	FSlateApplication::Get().UnregisterInputPreProcessor(InputProcessor);
+}
+
+void FMouseSlidingController::StartListeningForMouseEvents()
+{
+	// The user should first press U and then the LMB. If the LMB is pressed first, ignore.
+	// Other tools, such as the selection marquee in Curve Editor, use it and the interaction of those tools is hard to get right otherwise.
+	const bool bIsLeftMouseButtonDown = FSlateApplication::Get().GetPressedMouseButtons().Contains(EKeys::LeftMouseButton);
+	if (bIsLeftMouseButtonDown)
+	{
+		return;
+	}
+	
+	if (!SlidingState)
+	{
+		SlidingState.Emplace(*this);
+		OnStartSlidingDelegate.Broadcast();
+	}
+}
+
+void FMouseSlidingController::StopListeningForMouseEvents(ICursor& InCursor)
+{
+	if (!SlidingState)
+	{
+		return;
+	}
+
+	const bool bWasSliding = SlidingState->IsSliding();
+	if (bWasSliding)
+	{
+		StopSliding(InCursor);
+	}
+
+	SlidingState.Reset();
+}
+
+void FMouseSlidingController::StartSliding(const FVector2D& InInitialScreenLocation, ICursor& InCursor)
+{
+	SlidingState->InitialMouse = InInitialScreenLocation;
+	
+	const float SliderSize = MaxSlideWidthAttr.Get();
+	const float SliderHalfSize = SliderSize / 2.0;
+	
+	RECT ClipRect;
+	constexpr float LockHalfHeight = 12.f;
+	ClipRect.left = FMath::RoundToInt(SlidingState->InitialMouse->X - SliderHalfSize);
+	ClipRect.top = FMath::RoundToInt(SlidingState->InitialMouse->Y - LockHalfHeight);
+	ClipRect.right = FMath::TruncToInt(SlidingState->InitialMouse->X + SliderHalfSize);
+	ClipRect.bottom = FMath::TruncToInt(SlidingState->InitialMouse->Y + LockHalfHeight);
+	// Lock the mouse to the size of the virtual slider area to give the user feedback when they've moved the mouse far enough to reach -1 or 1.
+	InCursor.Lock(&ClipRect);
+
+	OnStartSlidingDelegate.Broadcast();
+}
+
+void FMouseSlidingController::StopSliding(ICursor& InCursor)
+{
+	SlidingState->InitialMouse.Reset();
+	InCursor.Lock(nullptr);
+	OnStopSlidingDelegate.Broadcast();
+}
+
+void FMouseSlidingController::UpdateSliding(const FVector2D& InScreenLocation)
+{
+	check(SlidingState->InitialMouse);
+	
+	const float SliderSize = MaxSlideWidthAttr.Get();
+	const float SliderHalfSize = SliderSize / 2.0;
+	const float InitialX = SlidingState->InitialMouse->X;
+	const float StartToMouse = InScreenLocation.X - InitialX;
+	const float ClampedStartToMouse = FMath::Clamp(StartToMouse, -SliderHalfSize, SliderHalfSize);
+	OnUpdateSlidingDelegate.Broadcast(ClampedStartToMouse / SliderHalfSize);
+}
+}
diff --git a/Engine/Plugins/Animation/TweeningUtils/Source/TweeningUtilsEditor/Private/Widgets/MVC/TweenMouseSlidingController.cpp b/Engine/Plugins/Animation/TweeningUtils/Source/TweeningUtilsEditor/Private/Widgets/MVC/TweenMouseSlidingController.cpp
new file mode 100644
--- /dev/null
+++ b/Engine/Plugins/Animation/TweeningUtils/Source/TweeningUtilsEditor/Private/Widgets/MVC/TweenMouseSlidingController.cpp
@@ -0,0 +1,49 @@
+﻿// Copyright Epic Games, Inc. All Rights Reserved.
+
+#include "Widgets/MVC/TweenMouseSlidingController.h"
+
+#include "HAL/IConsoleManager.h"
+#include "Misc/Attribute.h"
+#include "TweeningUtilsCommands.h"
+#include "TweeningUtilsStyle.h"
+#include "Widgets/TweenSliderStyle.h"
+#include "Widgets/MVC/Models/TweenModel.h"
+
+namespace UE::TweeningUtilsEditor
+{
+static TAutoConsoleVariable<float> CVarOverrideTweenSliderWidth(
+	TEXT("CurveEditor.SliderMouseWidth"),
+	-1.f,
+	TEXT("Specify positive value to override the sliding width. Non-positive will result in the default being used.")
+	);
+	
+FTweenMouseSlidingController::FTweenMouseSlidingController(
+	TAttribute<float> InMaxSlideWidthAttr,
+	const TSharedRef<FTweenModel>& InModel,
+	const TSharedRef<FUICommandList>& InCommandList,
+	TSharedPtr<FUICommandInfo> InDragSliderCommand
+	)
+	: FMouseSlidingController(
+		TAttribute<float>::CreateLambda([Attr = MoveTemp(InMaxSlideWidthAttr)]()
+		{
+			const float CVarOverride = CVarOverrideTweenSliderWidth.GetValueOnGameThread();
+			return CVarOverride <= 0.f ? Attr.Get() : CVarOverride;
+		}),
+		InCommandList, MoveTemp(InDragSliderCommand)
+		)
+	, TweenModel(InModel)
+{
+	OnStartSliding().AddSP(TweenModel, &FTweenModel::StartBlendOperation);
+	OnStopSliding().AddSP(TweenModel, &FTweenModel::StopBlendOperation);
+	OnUpdateSliding().AddSP(TweenModel, &FTweenModel::BlendValues);
+}
+
+FTweenMouseSlidingController::FTweenMouseSlidingController(const TSharedRef<FTweenModel>& InModel, const TSharedRef<FUICommandList>& InCommandList)
+	: FTweenMouseSlidingController(
+		FTweeningUtilsStyle::Get().GetWidgetStyle<FTweenSliderStyle>("TweenSlider").BarDimensions.X,
+		InModel,
+		InCommandList,
+		FTweeningUtilsCommands::Get().DragAnimSliderTool
+	)
+{}
+}
diff --git a/Engine/Plugins/Animation/TweeningUtils/Source/TweeningUtilsEditor/Public/Math/Abstraction/ITweenFunctionContainer.h b/Engine/Plugins/Animation/TweeningUtils/Source/TweeningUtilsEditor/Public/Math/Abstraction/ITweenFunctionContainer.h
--- a/Engine/Plugins/Animation/TweeningUtils/Source/TweeningUtilsEditor/Public/Math/Abstraction/ITweenFunctionContainer.h
+++ b/Engine/Plugins/Animation/TweeningUtils/Source/TweeningUtilsEditor/Public/Math/Abstraction/ITweenFunctionContainer.h
@@ -26,6 +26,10 @@ public:
 	virtual TSharedPtr<ITweenFunction> GetFunction(int32 InIndex) const = 0;
 	/** @return The number of functions contained. */
 	virtual int32 NumFunctions() const = 0;
+	/** @return The index of InFunction if found and INDEX_NONE if not. */
+	TWEENINGUTILSEDITOR_API int32 IndexOf(const TSharedRef<ITweenFunction>& InFunction) const;
+	/** @return Whether InIndex is a valid argument for GetFunction. */
+	bool IsValidIndex(int32 InIndex) const { return InIndex >= 0 && InIndex < NumFunctions(); }
 
 	/** @return The command that is used to select the function */
 	virtual TSharedPtr<FUICommandInfo> GetCommandForFunction(const TSharedRef<ITweenFunction>& InFunction) const = 0;
diff --git a/Engine/Plugins/Animation/TweeningUtils/Source/TweeningUtilsEditor/Public/TweeningUtilsCommands.h b/Engine/Plugins/Animation/TweeningUtils/Source/TweeningUtilsEditor/Public/TweeningUtilsCommands.h
--- a/Engine/Plugins/Animation/TweeningUtils/Source/TweeningUtilsEditor/Public/TweeningUtilsCommands.h
+++ b/Engine/Plugins/Animation/TweeningUtils/Source/TweeningUtilsEditor/Public/TweeningUtilsCommands.h
@@ -18,6 +18,9 @@ public:
 	TSharedPtr<FUICommandInfo> SetTweenSmoothRough;
 	TSharedPtr<FUICommandInfo> ToggleOvershootMode;
 
+	TSharedPtr<FUICommandInfo> DragAnimSliderTool;
+	TSharedPtr<FUICommandInfo> ChangeAnimSliderTool;
+
 	FTweeningUtilsCommands();
 
 	//~ Begin TCommands Interface
diff --git a/Engine/Plugins/Animation/TweeningUtils/Source/TweeningUtilsEditor/Public/Widgets/MVC/CycleFunctionController.h b/Engine/Plugins/Animation/TweeningUtils/Source/TweeningUtilsEditor/Public/Widgets/MVC/CycleFunctionController.h
new file mode 100644
--- /dev/null
+++ b/Engine/Plugins/Animation/TweeningUtils/Source/TweeningUtilsEditor/Public/Widgets/MVC/CycleFunctionController.h
@@ -0,0 +1,57 @@
+﻿// Copyright Epic Games, Inc. All Rights Reserved.
+
+#pragma once
+
+#include "Templates/SharedPointer.h"
+#include "Templates/UnrealTemplate.h"
+
+class FUICommandInfo;
+class FUICommandList;
+
+namespace UE::TweeningUtilsEditor
+{
+class FFunctionBasedTweenModel;
+class FTweenModel;
+class ITweenFunctionContainer;
+
+/** Cycles through functions on a tween model. */
+class TWEENINGUTILSEDITOR_API FCycleFunctionController : public FNoncopyable
+{
+public:
+	
+	/**
+	 * @param InTweenModel The model on which the functions will be changed.
+	 * @param InTweenFunctions Determines the functions through which you can cycle
+	 * @param InCommandList The command list to add / remove the command to / from.
+	 * @param InCycleCommand The command that cycles the function
+	 */
+	explicit FCycleFunctionController(
+		const TSharedRef<FFunctionBasedTweenModel>& InTweenModel,
+		const TSharedRef<ITweenFunctionContainer>& InTweenFunctions,
+		const TSharedRef<FUICommandList>& InCommandList,
+		TSharedPtr<FUICommandInfo> InCycleCommand
+		);
+	/** Version that defauls the cycle command to FTweenUtilsCommands::ChangeAnimSliderTool. */
+	explicit FCycleFunctionController(
+		const TSharedRef<FFunctionBasedTweenModel>& InTweenModel,
+		const TSharedRef<ITweenFunctionContainer>& InTweenFunctions,
+		const TSharedRef<FUICommandList>& InCommandList
+		); 
+	virtual ~FCycleFunctionController();
+	
+private:
+	
+	/** The model whose function is updated. */
+	const TSharedRef<FFunctionBasedTweenModel> TweenModel;
+	/** Holds the functions that can be cycled through. */
+	const TSharedRef<ITweenFunctionContainer> TweenFunctions;
+	
+	/** Used to bind and unbind the CycleCommand command. */
+	const TSharedRef<FUICommandList> CommandList;
+	/** Invokes CycleToNextFunction. */
+	const TSharedPtr<FUICommandInfo> CycleCommand;
+	
+	/** Sets the next tween function. */
+	void CycleToNextFunction();
+};
+}
diff --git a/Engine/Plugins/Animation/TweeningUtils/Source/TweeningUtilsEditor/Public/Widgets/MVC/MouseSlidingController.h b/Engine/Plugins/Animation/TweeningUtils/Source/TweeningUtilsEditor/Public/Widgets/MVC/MouseSlidingController.h
new file mode 100644
--- /dev/null
+++ b/Engine/Plugins/Animation/TweeningUtils/Source/TweeningUtilsEditor/Public/Widgets/MVC/MouseSlidingController.h
@@ -0,0 +1,114 @@
+﻿// Copyright Epic Games, Inc. All Rights Reserved.
+
+#pragma once
+
+#include "Delegates/Delegate.h"
+#include "Math/Vector2D.h"
+#include "Misc/Attribute.h"
+#include "Misc/Optional.h"
+#include "Templates/SharedPointer.h"
+#include "Templates/UnrealTemplate.h"
+
+class FUICommandInfo;
+class FUICommandList;
+class ICursor;
+
+namespace UE::TweeningUtilsEditor
+{
+class FMouseSlidingInputProcessor;
+class FTweenModel;
+
+/**
+ * Generic controller that allows you to move the mouse while a command is pressed.
+ * The sliding value is normalized to [-1,1] depending on whether the mouse was moved left or right.
+ * 
+ * The user must first press down the key chord for DragSliderCommand, and then the LMB. The user can press and release LMB multiple times while
+ * DragSliderCommand's key chord is pressed down.
+ * 
+ * During sliding, the mouse is locked to the rect of the virtual slider: the area around where the user started dragging.
+ * This is intended to give the user feedback when they've moved the mouse far enough to reach -1 or 1.
+ * 
+ * Example:
+ * - Suppose MaxSlideWidthAttr returns 200, and that InDragSliderCommand is bound to U.
+ * - The user presses U and moves the mouse
+ *		- to right by 50 Slate units > OnSliderMovedDelegate is invoked with is 0.5.
+ *		- to the right by 125 Slate units (exceeds the max sliding by 25 units) > OnSliderMovedDelegate is invoked with 1.0 (clamped).
+ *		- to the left by 25 Slate units > OnSliderMovedDelegate is invoked with -0.25.
+ */
+class TWEENINGUTILSEDITOR_API FMouseSlidingController : public FNoncopyable
+{
+	friend FMouseSlidingInputProcessor;
+public:
+	
+	/**
+	 * @param InMaxSlideWidthAttr The total size of the "invisible" slider. The mouse can be moved left and right InMaxSlideWidthAttr.Get / 2.0 before being clamped.
+	 * @param InCommandList The command list to add / remove the command to / from.
+	 * @param InDragSliderCommand The command that triggers detection of mouse movement.
+	 */
+	FMouseSlidingController(
+		TAttribute<float> InMaxSlideWidthAttr,
+		const TSharedRef<FUICommandList>& InCommandList,
+		TSharedPtr<FUICommandInfo> InDragSliderCommand
+		);
+	virtual ~FMouseSlidingController();
+
+	DECLARE_MULTICAST_DELEGATE_OneParam(FMoveSliderDelegate, float /*InNormalizedRange*/);
+	/** Invoked every tick the slider is moved. The event is throttled, so there's at most one invocation per tick. */
+	FMoveSliderDelegate& OnUpdateSliding() { return OnUpdateSlidingDelegate; }
+	/** Invoked when sliding starts. */
+	FSimpleMulticastDelegate& OnStartSliding() { return OnStartSlidingDelegate; }
+	/** Invoked when sliding stops. */
+	FSimpleMulticastDelegate& OnStopSliding() { return OnStopSlidingDelegate; }
+
+protected:
+
+	const TSharedRef<FUICommandList>& GetCommandList() const { return CommandList; }
+	
+private:
+	
+	/** The total size of the "invisible" slider. The mouse can be moved left and right by MaxSlideWidthAttr.Get() / 2.0 before being clamped. */
+	const TAttribute<float> MaxSlideWidthAttr;
+
+	/** Used to bind and unbind the DragSliderCommand command. */
+	const TSharedRef<FUICommandList> CommandList;
+	/** The command must be pressed down to start sliding. You also need to press LMB to start sliding. */
+	const TSharedPtr<FUICommandInfo> DragSliderCommand;
+
+	struct FSlidingState
+	{
+		/** Listens for DragSliderCommand and LMB going up. Orchestrates calls to StartSliding, StopSliding, and UpdateSliding. */
+		const TSharedRef<FMouseSlidingInputProcessor> InputProcessor;
+		/** Set while sliding is active. Unset while inactive. */
+		TOptional<FVector2D> InitialMouse;
+
+		explicit FSlidingState(FMouseSlidingController& InOwner);
+		~FSlidingState();
+
+		/** @return Whether sliding is currently active. */
+		bool IsSliding() const { return InitialMouse.IsSet(); }
+	};
+	/** Set while listening for mouse events, unset while not sliding. */
+	TOptional<FSlidingState> SlidingState;
+	
+	/** Invoked every tick the slider is moved. The event is throttled, so there's at most one invocation per tick. */
+	FMoveSliderDelegate OnUpdateSlidingDelegate;
+	/** Invoked when sliding starts. */
+	FSimpleMulticastDelegate OnStartSlidingDelegate;
+	/** Invoked when sliding stops. */
+	FSimpleMulticastDelegate OnStopSlidingDelegate;
+
+	/** Once DragSliderCommand has triggered, start listening for mouse down events. */
+	void StartListeningForMouseEvents();
+	/** Once DragSliderCommand is released, stop listening for mouse down events.  */
+	void StopListeningForMouseEvents(ICursor& InCursor);
+
+	/** Starts sliding from mouse position and locks the mouse; emits OnUpdateSliding calls. Called once DragSliderCommand and LMB are both down. */
+	void StartSliding(const FVector2D& InInitialScreenLocation, ICursor& InCursor);
+	/** Stops emiting emitting OnUpdateSliding calls. Unlocks the cursor. Called once either DragSliderCommand or LMB are released. */
+	void StopSliding(ICursor& InCursor);
+	/** Emits a OnUpdateSliding call. Called once per frame while sliding is active. */
+	void UpdateSliding(const FVector2D& InScreenLocation);
+};
+
+	
+}
diff --git a/Engine/Plugins/Animation/TweeningUtils/Source/TweeningUtilsEditor/Public/Widgets/MVC/TweenMouseSlidingController.h b/Engine/Plugins/Animation/TweeningUtils/Source/TweeningUtilsEditor/Public/Widgets/MVC/TweenMouseSlidingController.h
new file mode 100644
--- /dev/null
+++ b/Engine/Plugins/Animation/TweeningUtils/Source/TweeningUtilsEditor/Public/Widgets/MVC/TweenMouseSlidingController.h
@@ -0,0 +1,47 @@
+﻿// Copyright Epic Games, Inc. All Rights Reserved.
+
+#pragma once
+
+#include "MouseSlidingController.h"
+#include "Templates/SharedPointer.h"
+
+class FUICommandInfo;
+class FUICommandList;
+
+namespace UE::TweeningUtilsEditor
+{
+class FTweenModel;
+
+/** Specialization that sends blend values to a tween model based on how far the mouse is moved. */
+class TWEENINGUTILSEDITOR_API FTweenMouseSlidingController : public FMouseSlidingController
+{
+public:
+	
+	/**
+	 * Version that allows you to specify every detail
+	 * @param InMaxSlideWidthAttr The total size of the "invisible" slider. The mouse can be moved left and right InMaxSlideWidthAttr.Get / 2.0 before being clamped.
+	 * @param InModel The tween model that the sliding will update.
+	 * @param InCommandList The command list to bind / unbind the command with.
+	 * @param InDragSliderCommand The command that triggers detection of mouse movement.
+	 */
+	explicit FTweenMouseSlidingController(
+		TAttribute<float> InMaxSlideWidthAttr,
+		const TSharedRef<FTweenModel>& InModel,
+		const TSharedRef<FUICommandList>& InCommandList,
+		TSharedPtr<FUICommandInfo> InDragSliderCommand
+		);
+	/**
+	 * Version that defaults to using the commands from FTweeningUtilsCommands.
+	 * The slider width is based on the size of the slider in the toolbar(= value from FTweeningUtilsStyle).
+	 */
+	explicit FTweenMouseSlidingController(
+		const TSharedRef<FTweenModel>& InModel,
+		const TSharedRef<FUICommandList>& InCommandList
+	);
+
+private:
+	
+	/** Does the actual blending. */
+	const TSharedRef<FTweenModel>& TweenModel;
+};
+}
diff --git a/Engine/Plugins/Editor/CurveEditorTools/Source/CurveEditorTools/Private/Extensions/TweenEditorExtension.cpp b/Engine/Plugins/Editor/CurveEditorTools/Source/CurveEditorTools/Private/Extensions/TweenEditorExtension.cpp
--- a/Engine/Plugins/Editor/CurveEditorTools/Source/CurveEditorTools/Private/Extensions/TweenEditorExtension.cpp
+++ b/Engine/Plugins/Editor/CurveEditorTools/Source/CurveEditorTools/Private/Extensions/TweenEditorExtension.cpp
@@ -19,7 +19,7 @@ FTweenEditorExtension::FTweenEditorExtension(TWeakPtr<FCurveEditor> InCurveEdito
 
 void FTweenEditorExtension::BindCommands(TSharedRef<FUICommandList> InCommandList)
 {
-	GetOrCreateToolbarController(InCommandList);
+	InitControllers(InCommandList);
 }
 
 TSharedPtr<FExtender> FTweenEditorExtension::MakeToolbarExtender(const TSharedRef<FUICommandList>& InCommandList)
@@ -30,35 +30,44 @@ TSharedPtr<FExtender> FTweenEditorExtension::MakeToolbarExtender(const TSharedRe
 		[this, WeakCommandList = InCommandList.ToWeakPtr()](FToolBarBuilder& ToolbarBuilder)
 		{
 			const TSharedPtr<FUICommandList> CommandListPin = WeakCommandList.Pin();
-			if (ensure(CommandListPin) && GetOrCreateToolbarController(CommandListPin.ToSharedRef()))
+			if (ensure(CommandListPin) && InitControllers(CommandListPin.ToSharedRef()))
 			{
 				ToolbarBuilder.BeginSection("Tween");
-				ToolbarController->AddToToolbar(ToolbarBuilder);
+				TweenControllers->ToolbarController.AddToToolbar(ToolbarBuilder);
 				ToolbarBuilder.EndSection();
 			}
 		}));
 	
 	return Extender;
 }
+	
+FTweenEditorExtension::FTweenControllers::FTweenControllers(
+	const TSharedRef<TweeningUtilsEditor::FFunctionBasedTweenModel>& InTweenModel,
+	const TSharedRef<FUICommandList>& InCommandList,
+	const TSharedRef<TweeningUtilsEditor::ITweenFunctionContainer>& InTweenFunctions
+	)
+	: ToolbarController(InTweenModel, InCommandList, InTweenFunctions)
+	, CycleFunctionController(InTweenModel, InTweenFunctions, InCommandList)
+	, MouseSlidingController(InTweenModel, InCommandList)
+{}
 
-TSharedPtr<TweeningUtilsEditor::FTweenToolbarController> FTweenEditorExtension::GetOrCreateToolbarController(TSharedRef<FUICommandList> InCommandList)
+bool FTweenEditorExtension::InitControllers(TSharedRef<FUICommandList> InCommandList)
 {
-	if (ToolbarController)
+	if (TweenControllers)
 	{
-		return ToolbarController;
+		return true;
 	}
 
 	if (const TSharedPtr<FCurveEditor> CurveEditorPin = WeakCurveEditor.Pin(); ensure(CurveEditorPin))
 	{
 		using namespace TweeningUtilsEditor;
-		
 		const TSharedRef<TTransactionalTweenModelProxy<FCurveTweenModel>> TweenModel = MakeShared<TTransactionalTweenModelProxy<FCurveTweenModel>>(
 			WeakCurveEditor,
 			BlendFunctionContainer->GetFunction(0).ToSharedRef()
 			);
-		ToolbarController = MakeShared<FTweenToolbarController>(TweenModel, InCommandList, BlendFunctionContainer);
+		TweenControllers.Emplace(TweenModel, InCommandList, BlendFunctionContainer);
+		return true;
 	}
-	
-	return ToolbarController;
+	return false;
 }
 }
diff --git a/Engine/Plugins/Editor/CurveEditorTools/Source/CurveEditorTools/Private/Extensions/TweenEditorExtension.h b/Engine/Plugins/Editor/CurveEditorTools/Source/CurveEditorTools/Private/Extensions/TweenEditorExtension.h
--- a/Engine/Plugins/Editor/CurveEditorTools/Source/CurveEditorTools/Private/Extensions/TweenEditorExtension.h
+++ b/Engine/Plugins/Editor/CurveEditorTools/Source/CurveEditorTools/Private/Extensions/TweenEditorExtension.h
@@ -3,12 +3,21 @@
 #pragma once
 
 #include "ICurveEditorExtension.h"
+#include "Misc/Optional.h"
 #include "Templates/SharedPointer.h"
+#include "Widgets/MVC/CycleFunctionController.h"
+#include "Widgets/MVC/TweenMouseSlidingController.h"
+#include "Widgets/MVC/TweenToolbarController.h"
 
 class FCurveEditor;
 class FExtender;
-namespace UE::TweeningUtilsEditor { class FBuiltInTweenFunctionContainer; }
-namespace UE::TweeningUtilsEditor { class FTweenToolbarController; }
+
+namespace UE::TweeningUtilsEditor
+{
+	class FBuiltInTweenFunctionContainer;
+	class FFunctionBasedTweenModel;
+	class ITweenFunctionContainer;
+}
 
 namespace UE::CurveEditor
 {
@@ -31,11 +40,29 @@ private:
 
 	/** Holds the functions that ToolbarController shows: just the default, built-in ones.*/
 	const TSharedRef<TweeningUtilsEditor::FBuiltInTweenFunctionContainer> BlendFunctionContainer;
-	/** Created by BindCommands. */
-	TSharedPtr<TweeningUtilsEditor::FTweenToolbarController> ToolbarController;
+
+	struct FTweenControllers
+	{
+		/** Manages the toolbar widget */
+		const TweeningUtilsEditor::FTweenToolbarController ToolbarController;
+
+		/** Cycles functions (Shift + U). */
+		const TweeningUtilsEditor::FCycleFunctionController CycleFunctionController;
+
+		/** Allows indirect movement of the slider by using u + Move Mouse */
+		const TweeningUtilsEditor::FTweenMouseSlidingController MouseSlidingController;
+		
+		explicit FTweenControllers(
+			const TSharedRef<TweeningUtilsEditor::FFunctionBasedTweenModel>& InTweenModel,
+			const TSharedRef<FUICommandList>& InCommandList,
+			const TSharedRef<TweeningUtilsEditor::ITweenFunctionContainer>& InTweenFunctions
+			);
+	};
+	/** Created by InitControllers. */
+	TOptional<FTweenControllers> TweenControllers;
 
 	/** Makes sure that ToolbarController is created. */
-	TSharedPtr<TweeningUtilsEditor::FTweenToolbarController> GetOrCreateToolbarController(TSharedRef<FUICommandList> InCommandList);
+	bool InitControllers(TSharedRef<FUICommandList> InCommandList);
 };
 }
 
