Curve Editor: Ctrl+U hotkey for toggling overshoot mode

diff --git a/Engine/Plugins/Animation/TweeningUtils/Source/TweeningUtilsEditor/Private/TweeningUtilsCommands.cpp b/Engine/Plugins/Animation/TweeningUtils/Source/TweeningUtilsEditor/Private/TweeningUtilsCommands.cpp
--- a/Engine/Plugins/Animation/TweeningUtils/Source/TweeningUtilsEditor/Private/TweeningUtilsCommands.cpp
+++ b/Engine/Plugins/Animation/TweeningUtils/Source/TweeningUtilsEditor/Private/TweeningUtilsCommands.cpp
@@ -29,7 +29,7 @@ void FTweeningUtilsCommands::RegisterCommands()
 	UI_COMMAND(SetTweenBlendRelative, "Set Blend Relative", "Sets the tween mode to Blend Relative", EUserInterfaceActionType::RadioButton, FInputChord() );
 	UI_COMMAND(SetTweenBlendEase, "Set Blend Ease", "Sets the tween mode to Blend Ease", EUserInterfaceActionType::RadioButton, FInputChord() );
 	UI_COMMAND(SetTweenSmoothRough, "Set Smooth / Rough", "Sets the tween mode to Smooth / Rough", EUserInterfaceActionType::RadioButton, FInputChord() );
-	UI_COMMAND(ToggleOvershootMode, "Overshoot mode", "Overshoot extends the available range over 100%", EUserInterfaceActionType::ToggleButton, FInputChord() );
+	UI_COMMAND(ToggleOvershootMode, "Overshoot mode", "Overshoot extends the available range over 100%", EUserInterfaceActionType::ToggleButton, FInputChord(EKeys::U, EModifierKey::Control) );
 	
 	UI_COMMAND(DragAnimSliderTool, "Drag Anim Slider Tool", "Drag the anim slider", EUserInterfaceActionType::Button, FInputChord(EKeys::U));
 	UI_COMMAND(ChangeAnimSliderTool, "Change Anim Slider Tool", "Go to the next blend function", EUserInterfaceActionType::Button, FInputChord(EKeys::U, EModifierKey::Shift));
