Curve Editor: Visually move tweening slider while indirectly moving it with U+LMB.
- FTweenMouseSlidingController now exposes the indirect value
- The indirect value is injected into FTweenToolbarController's toolbar widget, which injects it to STweenSlider
- STweenSlider now has a TAttribute for overriding the slider position (only used if you're not dragging the slider)

diff --git a/Engine/Plugins/Animation/TweeningUtils/Source/TweeningUtilsEditor/Private/Widgets/MVC/STweenView.cpp b/Engine/Plugins/Animation/TweeningUtils/Source/TweeningUtilsEditor/Private/Widgets/MVC/STweenView.cpp
--- a/Engine/Plugins/Animation/TweeningUtils/Source/TweeningUtilsEditor/Private/Widgets/MVC/STweenView.cpp
+++ b/Engine/Plugins/Animation/TweeningUtils/Source/TweeningUtilsEditor/Private/Widgets/MVC/STweenView.cpp
@@ -30,6 +30,7 @@ void STweenView::Construct(const FArguments& InArgs, const TSharedRef<FTweenMode
 			SNew(STweenSlider)
 			.SliderIcon(InArgs._SliderIcon)
 			.SliderColor(InArgs._SliderColor)
+			.OverrideSliderPosition(InArgs._OverrideSliderPosition)
 			.ScaleRenderMode(this, &STweenView::GetBarRenderMode)
 			.OnSliderDragStarted(this, &STweenView::OnDragStarted)
 			.OnSliderDragEnded(this, &STweenView::OnDragEnded)
diff --git a/Engine/Plugins/Animation/TweeningUtils/Source/TweeningUtilsEditor/Private/Widgets/MVC/TweenMouseSlidingController.cpp b/Engine/Plugins/Animation/TweeningUtils/Source/TweeningUtilsEditor/Private/Widgets/MVC/TweenMouseSlidingController.cpp
--- a/Engine/Plugins/Animation/TweeningUtils/Source/TweeningUtilsEditor/Private/Widgets/MVC/TweenMouseSlidingController.cpp
+++ b/Engine/Plugins/Animation/TweeningUtils/Source/TweeningUtilsEditor/Private/Widgets/MVC/TweenMouseSlidingController.cpp
@@ -34,8 +34,8 @@ FTweenMouseSlidingController::FTweenMouseSlidingController(
 	, TweenModel(InModel)
 {
 	OnStartSliding().AddSP(TweenModel, &FTweenModel::StartBlendOperation);
-	OnStopSliding().AddSP(TweenModel, &FTweenModel::StopBlendOperation);
-	OnUpdateSliding().AddSP(TweenModel, &FTweenModel::BlendValues);
+	OnStopSliding().AddRaw(this, &FTweenMouseSlidingController::OnSliderEndMove);
+	OnUpdateSliding().AddRaw(this, &FTweenMouseSlidingController::OnSliderMove);
 }
 
 FTweenMouseSlidingController::FTweenMouseSlidingController(const TSharedRef<FTweenModel>& InModel, const TSharedRef<FUICommandList>& InCommandList)
@@ -46,4 +46,16 @@ FTweenMouseSlidingController::FTweenMouseSlidingController(const TSharedRef<FTwe
 		FTweeningUtilsCommands::Get().DragAnimSliderTool
 	)
 {}
+
+void FTweenMouseSlidingController::OnSliderEndMove()
+{
+	CurrentSliderPosition.Reset();
+	TweenModel->StopBlendOperation();
+}
+	
+void FTweenMouseSlidingController::OnSliderMove(float InValue)
+{
+	CurrentSliderPosition = InValue;
+	TweenModel->BlendValues(InValue);
+}
 }
diff --git a/Engine/Plugins/Animation/TweeningUtils/Source/TweeningUtilsEditor/Private/Widgets/MVC/TweenToolbarController.cpp b/Engine/Plugins/Animation/TweeningUtils/Source/TweeningUtilsEditor/Private/Widgets/MVC/TweenToolbarController.cpp
--- a/Engine/Plugins/Animation/TweeningUtils/Source/TweeningUtilsEditor/Private/Widgets/MVC/TweenToolbarController.cpp
+++ b/Engine/Plugins/Animation/TweeningUtils/Source/TweeningUtilsEditor/Private/Widgets/MVC/TweenToolbarController.cpp
@@ -54,7 +54,7 @@ FTweenToolbarController::~FTweenToolbarController()
 	CommandList->UnmapAction(FTweeningUtilsCommands::Get().ToggleOvershootMode);
 }
 
-void FTweenToolbarController::AddToToolbar(FToolBarBuilder& ToolBarBuilder) const
+void FTweenToolbarController::AddToToolbar(FToolBarBuilder& ToolBarBuilder, TAttribute<TOptional<float>> InOverrideSliderPosition) const
 {
 	ToolBarBuilder.SetLabelVisibility(EVisibility::Visible);
 	ToolBarBuilder.AddComboButton(
@@ -70,6 +70,7 @@ void FTweenToolbarController::AddToToolbar(FToolBarBuilder& ToolBarBuilder) cons
 		SNew(STweenView, TweenModel)
 		.SliderIcon(TAttribute<const FSlateBrush*>::CreateRaw(this, &FTweenToolbarController::GetIconForSlider))
 		.SliderColor(TAttribute<FLinearColor>::CreateRaw(this, &FTweenToolbarController::GetColorForSlider))
+		.OverrideSliderPosition(MoveTemp(InOverrideSliderPosition))
 		);
 	ToolBarBuilder.AddToolBarButton(FTweeningUtilsCommands::Get().ToggleOvershootMode);
 }
diff --git a/Engine/Plugins/Animation/TweeningUtils/Source/TweeningUtilsEditor/Private/Widgets/STweenSlider.cpp b/Engine/Plugins/Animation/TweeningUtils/Source/TweeningUtilsEditor/Private/Widgets/STweenSlider.cpp
--- a/Engine/Plugins/Animation/TweeningUtils/Source/TweeningUtilsEditor/Private/Widgets/STweenSlider.cpp
+++ b/Engine/Plugins/Animation/TweeningUtils/Source/TweeningUtilsEditor/Private/Widgets/STweenSlider.cpp
@@ -2,6 +2,7 @@
 
 #include "Widgets/STweenSlider.h"
 
+#include "Framework/DelayedDrag.h"
 #include "Internationalization/Internationalization.h"
 #include "Internationalization/Text.h"
 #include "Widgets/SToolTip.h"
@@ -58,11 +59,19 @@ struct FSliderWidgetData
 	TSharedPtr<SToolTip> ToolTip;
 };
 
-/** Converts [0,1] range to [-1,1] range. */
-static float SliderValueToBlendValue(float InSliderValue)
+/** Converts [0,1] to [-1,1] */
+static float AsymmetricToSymmetric(float InValue) { return InValue * 2.0 - 1.0; }
+/** Converts [-1,1] to [0,1] */
+static float SymmetricToAsymmetric(float InValue) { return (InValue + 1) / 2.0; }
+	
+static float GetSliderPosition(const FTweenWidgetArgs& InWidgetArgs, float InFallback)
 {
-	return InSliderValue * 2.0 - 1.0;
+	const bool bHasOverride = InWidgetArgs.OverrideSliderPositionAttr.IsBound() || InWidgetArgs.OverrideSliderPositionAttr.IsSet();
+	const TOptional<float> OverrideValue = bHasOverride
+		? InWidgetArgs.OverrideSliderPositionAttr.Get() : TOptional<float>();
+	return OverrideValue ? SymmetricToAsymmetric(*OverrideValue) : InFallback;
 }
+
 	
 void STweenSlider::Construct(const FArguments& InArgs)
 {
@@ -80,6 +89,7 @@ void STweenSlider::Construct(const FArguments& InArgs)
 	WidgetArgs.SliderIconAttr = InArgs._SliderIcon;
 	WidgetArgs.SliderColor = InArgs._SliderColor;
 	WidgetArgs.ScaleRenderModeAttr = InArgs._ScaleRenderMode;
+	WidgetArgs.OverrideSliderPositionAttr = InArgs._OverrideSliderPosition;
 
 	Pimpl->ToolTip = SNew(SToolTip)
 		.Text(this, &STweenSlider::GetToolTipText)
@@ -228,7 +238,7 @@ void STweenSlider::OnFinishedPointerInput()
 		// When OnFinishedPointerInput is called, we know the final mouse position for the frame.
 		OnSliderChangedDelegate.ExecuteIfBound(
 			// Convert [0,1] range to [-1,1] range.
-			SliderValueToBlendValue(Pimpl->TargetSliderPosition)
+			AsymmetricToSymmetric(Pimpl->TargetSliderPosition)
 			);
 	}
 }
@@ -276,18 +286,22 @@ int32 STweenSlider::OnPaint(
 	TArray<float> NormalizedPositions;
 	FTweenSliderDrawArgs DrawArgs;
 	
-	GetBarGeometry(InAllottedGeometry, WidgetArgs, DrawArgs.BarArea);
-	GetSliderButtonGeometry(Pimpl->AnimatedCurrentSliderPosition, InAllottedGeometry, WidgetArgs, DrawArgs.SliderArea, DrawArgs.IconArea);
-	GetDrawnPointGeometry(
-		InAllottedGeometry, WidgetArgs, HoverState, bIsMouseButtonDown, DrawArgs.Points, DrawArgs.PointTypes, NormalizedPositions
-		);
-	
 	DrawArgs.bDrawButtonPressed = bIsDragging
 		|| (bIsMouseButtonDown && HoverState.bIsSliderHovered)
 		|| (bIsMouseButtonDown && HoverState.HoveredPointIndex);
 	DrawArgs.bIsDragging = bIsDragging;
 	DrawArgs.HoverState = HoverState;
 	
+	const float SliderPosition = DrawArgs.bDrawButtonPressed
+		? Pimpl->AnimatedCurrentSliderPosition : GetSliderPosition(WidgetArgs, Pimpl->AnimatedCurrentSliderPosition);
+	GetBarGeometry(InAllottedGeometry, WidgetArgs, DrawArgs.BarArea);
+	GetSliderButtonGeometry(
+		SliderPosition, InAllottedGeometry, WidgetArgs, DrawArgs.SliderArea, DrawArgs.IconArea
+		);
+	GetDrawnPointGeometry(
+		InAllottedGeometry, WidgetArgs, HoverState, bIsMouseButtonDown, DrawArgs.Points, DrawArgs.PointTypes, NormalizedPositions
+		);
+	
 	if (bIsDragging)
 	{
 		GetPassedPointStates(NormalizedPositions, Pimpl->AnimatedCurrentSliderPosition, DrawArgs.PassedPoints);
@@ -336,7 +350,7 @@ void STweenSlider::HandlePickValueOfCurrentlyHoveredPoint(const FGeometry& InGeo
 	if (PointValues.IsValidIndex(HoverIndex))
 	{
 		OnPointValuePickedDelegate.ExecuteIfBound(
-			SliderValueToBlendValue(PointValues[HoverIndex])
+			AsymmetricToSymmetric(PointValues[HoverIndex])
 			);
 	}
 }
@@ -354,7 +368,7 @@ FText STweenSlider::GetToolTipText() const
 		const bool bRoundUpper = 1 - Value < RoundThreshold;
 		const float AdjustedValue = bRoundLower ? 0.0 : bRoundUpper ? 1.0 : Value;
 		
-		return MapSliderValueToBlendValueDelegate.Execute(SliderValueToBlendValue(AdjustedValue));
+		return MapSliderValueToBlendValueDelegate.Execute(AsymmetricToSymmetric(AdjustedValue));
 	};
 	
 	const int32 HoveredIndex = Pimpl->HoverState.HoveredPointIndex.Get(INDEX_NONE);
diff --git a/Engine/Plugins/Animation/TweeningUtils/Source/TweeningUtilsEditor/Private/Widgets/TweenSliderDrawUtils.h b/Engine/Plugins/Animation/TweeningUtils/Source/TweeningUtilsEditor/Private/Widgets/TweenSliderDrawUtils.h
--- a/Engine/Plugins/Animation/TweeningUtils/Source/TweeningUtilsEditor/Private/Widgets/TweenSliderDrawUtils.h
+++ b/Engine/Plugins/Animation/TweeningUtils/Source/TweeningUtilsEditor/Private/Widgets/TweenSliderDrawUtils.h
@@ -88,6 +88,9 @@ struct FTweenWidgetArgs
 	TAttribute<FLinearColor> SliderColor;
 	/** Affects how the normalized values are supposed to be interpreted */
 	TAttribute<ETweenScaleMode> ScaleRenderModeAttr;
+	
+	/** If set, an indication where to position the slider. Range [-1,1]. If unset, defaults to 0. Ignored if the user is dragging the slider. */
+	TAttribute<TOptional<float>> OverrideSliderPositionAttr;
 };
 
 /** Gets geometry of the background bar. */
diff --git a/Engine/Plugins/Animation/TweeningUtils/Source/TweeningUtilsEditor/Public/Widgets/MVC/Models/TweenModel.h b/Engine/Plugins/Animation/TweeningUtils/Source/TweeningUtilsEditor/Public/Widgets/MVC/Models/TweenModel.h
--- a/Engine/Plugins/Animation/TweeningUtils/Source/TweeningUtilsEditor/Public/Widgets/MVC/Models/TweenModel.h
+++ b/Engine/Plugins/Animation/TweeningUtils/Source/TweeningUtilsEditor/Public/Widgets/MVC/Models/TweenModel.h
@@ -60,8 +60,6 @@ private:
 	 * - True: -1.0 to 1.0 maps to -200% to +200%
 	 */
 	ETweenScaleMode ScaleMode = ETweenScaleMode::Normalized;
-	
-	TOptional<float> CurrentBlendValue;
 };
 }
 
diff --git a/Engine/Plugins/Animation/TweeningUtils/Source/TweeningUtilsEditor/Public/Widgets/MVC/STweenView.h b/Engine/Plugins/Animation/TweeningUtils/Source/TweeningUtilsEditor/Public/Widgets/MVC/STweenView.h
--- a/Engine/Plugins/Animation/TweeningUtils/Source/TweeningUtilsEditor/Public/Widgets/MVC/STweenView.h
+++ b/Engine/Plugins/Animation/TweeningUtils/Source/TweeningUtilsEditor/Public/Widgets/MVC/STweenView.h
@@ -5,6 +5,8 @@
 #include "Widgets/DeclarativeSyntaxSupport.h"
 #include "Widgets/SCompoundWidget.h"
 
+template<typename OptionalType> struct TOptional;
+
 namespace UE::TweeningUtilsEditor
 {
 class FTweenModel;
@@ -23,6 +25,8 @@ public:
 		SLATE_ATTRIBUTE(const FSlateBrush*, SliderIcon)
 		/** The main color. It tints the slider button and the points. */
 		SLATE_ATTRIBUTE(FLinearColor, SliderColor)
+		/** If set, an indication where to position the slider. Range [-1,1]. If unset, defaults to 0. Ignored if the user is dragging the slider. */
+		SLATE_ATTRIBUTE(TOptional<float>, OverrideSliderPosition)
 	SLATE_END_ARGS()
 
 	void Construct(const FArguments& InArgs, const TSharedRef<FTweenModel>& InTweenModel);
diff --git a/Engine/Plugins/Animation/TweeningUtils/Source/TweeningUtilsEditor/Public/Widgets/MVC/TweenMouseSlidingController.h b/Engine/Plugins/Animation/TweeningUtils/Source/TweeningUtilsEditor/Public/Widgets/MVC/TweenMouseSlidingController.h
--- a/Engine/Plugins/Animation/TweeningUtils/Source/TweeningUtilsEditor/Public/Widgets/MVC/TweenMouseSlidingController.h
+++ b/Engine/Plugins/Animation/TweeningUtils/Source/TweeningUtilsEditor/Public/Widgets/MVC/TweenMouseSlidingController.h
@@ -2,6 +2,7 @@
 
 #pragma once
 
+#include "Misc/Optional.h"
 #include "MouseSlidingController.h"
 #include "Templates/SharedPointer.h"
 
@@ -39,9 +40,18 @@ public:
 		const TSharedRef<FUICommandList>& InCommandList
 	);
 
+	/** @return The current slider position that is being driven by the mouse. Range [-1,1]. */
+	const TOptional<float>& GetCurrentSliderPosition() const { return CurrentSliderPosition; }
+
 private:
 	
 	/** Does the actual blending. */
-	const TSharedRef<FTweenModel>& TweenModel;
+	const TSharedRef<FTweenModel> TweenModel;
+
+	/** The current slider position that is being driven by the mouse. Range [-1,1]. */
+	TOptional<float> CurrentSliderPosition;
+
+	void OnSliderEndMove();
+	void OnSliderMove(float InValue);
 };
 }
diff --git a/Engine/Plugins/Animation/TweeningUtils/Source/TweeningUtilsEditor/Public/Widgets/MVC/TweenToolbarController.h b/Engine/Plugins/Animation/TweeningUtils/Source/TweeningUtilsEditor/Public/Widgets/MVC/TweenToolbarController.h
--- a/Engine/Plugins/Animation/TweeningUtils/Source/TweeningUtilsEditor/Public/Widgets/MVC/TweenToolbarController.h
+++ b/Engine/Plugins/Animation/TweeningUtils/Source/TweeningUtilsEditor/Public/Widgets/MVC/TweenToolbarController.h
@@ -2,6 +2,7 @@
 
 #pragma once
 
+#include "Misc/Attribute.h"
 #include "Templates/SharedPointer.h"
 #include "Templates/UnrealTemplate.h"
 
@@ -12,6 +13,7 @@ class SWidget;
 struct FLinearColor;
 struct FSlateBrush;
 struct FSlateIcon;
+template<typename OptionalType> struct TOptional;
 
 namespace UE::TweeningUtilsEditor
 {
@@ -44,8 +46,14 @@ public:
 		);
 	~FTweenToolbarController();
 	
-	/** Extends the toolbar. Safe to call multiple times (although typically you'll call this exactly once). */
-	void AddToToolbar(FToolBarBuilder& ToolbarBuilder) const;
+	/**
+	 * Extends the toolbar. Safe to call multiple times (although typically you'll call this exactly once).
+	 * 
+	 * @param ToolbarBuilder The toolbar to add the widgets to
+	 * @param InOverrideSliderPosition Optional override for the where the slider is visually placed (does not affect functionality). Range [-1,1].
+	 *	Useful e.g. in conjunction with FTweenMouseSlidingController (while pressing U+LMB, give the user a visual association).
+	 */
+	void AddToToolbar(FToolBarBuilder& ToolbarBuilder, TAttribute<TOptional<float>> InOverrideSliderPosition = {}) const;
 
 private:
 
diff --git a/Engine/Plugins/Animation/TweeningUtils/Source/TweeningUtilsEditor/Public/Widgets/STweenSlider.h b/Engine/Plugins/Animation/TweeningUtils/Source/TweeningUtilsEditor/Public/Widgets/STweenSlider.h
--- a/Engine/Plugins/Animation/TweeningUtils/Source/TweeningUtilsEditor/Public/Widgets/STweenSlider.h
+++ b/Engine/Plugins/Animation/TweeningUtils/Source/TweeningUtilsEditor/Public/Widgets/STweenSlider.h
@@ -3,12 +3,14 @@
 #pragma once
 
 #include "ETweenScaleMode.h"
+#include "Styling/SlateWidgetStyleAsset.h" // Required by SLATE_STYLE_ARGUMENT
 #include "TweeningUtilsStyle.h"
 #include "TweenSliderStyle.h"
-#include "Framework/DelayedDrag.h"
 #include "Widgets/DeclarativeSyntaxSupport.h"
 #include "Widgets/SLeafWidget.h"
 
+template<typename T> struct TOptional;
+
 namespace UE::TweeningUtilsEditor
 {
 struct FSliderWidgetData;
@@ -26,9 +28,9 @@ struct FSliderWidgetData;
 class TWEENINGUTILSEDITOR_API STweenSlider : public SLeafWidget
 {
 public:
-
-	DECLARE_DELEGATE_OneParam(FSliderChangedDelegate, float);
-	DECLARE_DELEGATE_RetVal_OneParam(float, FMapSliderValueToBlendValue, float);
+	
+	DECLARE_DELEGATE_OneParam(FSliderChangedDelegate, float /*Value*/); // Value's range is [-1,1].
+	DECLARE_DELEGATE_RetVal_OneParam(float, FMapSliderValueToBlendValue, float /*Value*/); // Value's range is [-1,1].
 
 	SLATE_BEGIN_ARGS(STweenSlider)
 		: _Style(&FTweeningUtilsStyle::Get().GetWidgetStyle<FTweenSliderStyle>("TweenSlider"))
@@ -50,6 +52,9 @@ public:
 		SLATE_ATTRIBUTE(FLinearColor, SliderColor)
 		/** Affects how the scale is rendered. Events keep returning the -1 to 1 mode - you should compute the scale yourself.*/
 		SLATE_ATTRIBUTE(ETweenScaleMode, ScaleRenderMode)
+
+		/** If set, an indication where to position the slider. If unset, defaults to 0. Ignored if the user is dragging the slider. Range [-1,1]. */
+		SLATE_ATTRIBUTE(TOptional<float>, OverrideSliderPosition)
 		
 		/** Invoked when the slider begins being dragged. */
 		SLATE_EVENT(FSimpleDelegate, OnSliderDragStarted)
diff --git a/Engine/Plugins/Editor/CurveEditorTools/Source/CurveEditorTools/Private/Extensions/TweenEditorExtension.cpp b/Engine/Plugins/Editor/CurveEditorTools/Source/CurveEditorTools/Private/Extensions/TweenEditorExtension.cpp
--- a/Engine/Plugins/Editor/CurveEditorTools/Source/CurveEditorTools/Private/Extensions/TweenEditorExtension.cpp
+++ b/Engine/Plugins/Editor/CurveEditorTools/Source/CurveEditorTools/Private/Extensions/TweenEditorExtension.cpp
@@ -24,6 +24,7 @@ void FTweenEditorExtension::BindCommands(TSharedRef<FUICommandList> InCommandLis
 
 TSharedPtr<FExtender> FTweenEditorExtension::MakeToolbarExtender(const TSharedRef<FUICommandList>& InCommandList)
 {
+	using namespace TweeningUtilsEditor;
 	TSharedRef<FExtender> Extender = MakeShared<FExtender>();
 	
 	Extender->AddToolBarExtension("Tangents", EExtensionHook::After, InCommandList, FToolBarExtensionDelegate::CreateLambda(
@@ -32,8 +33,15 @@ TSharedPtr<FExtender> FTweenEditorExtension::MakeToolbarExtender(const TSharedRe
 			const TSharedPtr<FUICommandList> CommandListPin = WeakCommandList.Pin();
 			if (ensure(CommandListPin) && InitControllers(CommandListPin.ToSharedRef()))
 			{
+				// Visually move the slider widget when the user uses the U+LMB button command to indirectly move the mouse
+				TAttribute<TOptional<float>> OverrideSliderPosition =
+					TAttribute<TOptional<float>>::CreateLambda([this]()
+					{
+						return TweenControllers->MouseSlidingController.GetCurrentSliderPosition();
+					});
+				
 				ToolbarBuilder.BeginSection("Tween");
-				TweenControllers->ToolbarController.AddToToolbar(ToolbarBuilder);
+				TweenControllers->ToolbarController.AddToToolbar(ToolbarBuilder, MoveTemp(OverrideSliderPosition));
 				ToolbarBuilder.EndSection();
 			}
 		}));
