Curve Editor: Reorder blend functions in tweening toolbar combo box.

diff --git a/Engine/Plugins/Animation/TweeningUtils/Source/TweeningUtilsEditor/Private/Math/Abstraction/BuiltInTweenFunctionContainer.cpp b/Engine/Plugins/Animation/TweeningUtils/Source/TweeningUtilsEditor/Private/Math/Abstraction/BuiltInTweenFunctionContainer.cpp
--- a/Engine/Plugins/Animation/TweeningUtils/Source/TweeningUtilsEditor/Private/Math/Abstraction/BuiltInTweenFunctionContainer.cpp
+++ b/Engine/Plugins/Animation/TweeningUtils/Source/TweeningUtilsEditor/Private/Math/Abstraction/BuiltInTweenFunctionContainer.cpp
@@ -9,24 +9,22 @@
 
 namespace UE::TweeningUtilsEditor
 {
-static TOptional<EBlendFunction> GetBlendFunctionType(
-	const TSharedRef<ITweenFunction>& InFunctionToFind, const TArray<TSharedRef<FBuiltInTweenFunction>>& InFunctions
-	)
+FBuiltInTweenFunctionContainer::FBuiltInTweenFunctionContainer()
+{
+	ForEachBlendFunction([this](EBlendFunction BlendFunction){ Functions.Add(MakeShared<FBuiltInTweenFunction>(BlendFunction)); });
+}
+
+TOptional<EBlendFunction> FBuiltInTweenFunctionContainer::GetBlendFunctionType(const TSharedRef<ITweenFunction>& InFunctionToFind) const
 {
-	const int32 Index = InFunctions.IndexOfByPredicate([&InFunctionToFind](const TSharedRef<FBuiltInTweenFunction>& Function)
+	const int32 Index = Functions.IndexOfByPredicate([&InFunctionToFind](const TSharedRef<FBuiltInTweenFunction>& Function)
 	{
 		return &InFunctionToFind.Get() == &Function.Get();
 	});
 	// API user can theoretically pass in ANY function so handle that case
-	return ensureMsgf(InFunctions.IsValidIndex(Index), TEXT("You passed a function not managed by this container"))
-		? InFunctions[Index]->GetBlendFunction()
+	return ensureMsgf(Functions.IsValidIndex(Index), TEXT("You passed a function not managed by this container"))
+		? Functions[Index]->GetBlendFunction()
 		: TOptional<EBlendFunction>();
 }
-	
-FBuiltInTweenFunctionContainer::FBuiltInTweenFunctionContainer()
-{
-	ForEachBlendFunction([this](EBlendFunction BlendFunction){ Functions.Add(MakeShared<FBuiltInTweenFunction>(BlendFunction)); });
-}
 
 void FBuiltInTweenFunctionContainer::ForEachFunction(const TFunctionRef<void(const TSharedRef<ITweenFunction>&)> InConsumer) const
 {
@@ -43,31 +41,31 @@ TSharedPtr<ITweenFunction> FBuiltInTweenFunctionContainer::GetFunction(int32 InI
 
 TSharedPtr<FUICommandInfo> FBuiltInTweenFunctionContainer::GetCommandForFunction(const TSharedRef<ITweenFunction>& InFunction) const
 {
-	const TOptional<EBlendFunction> Function = GetBlendFunctionType(InFunction, Functions);
+	const TOptional<EBlendFunction> Function = GetBlendFunctionType(InFunction);
 	return Function ? TweeningUtilsEditor::GetCommandForBlendFunction(*Function) : nullptr;
 }
 
 const FSlateBrush* FBuiltInTweenFunctionContainer::GetIconForFunction(const TSharedRef<ITweenFunction>& InFunction) const
 {
-	const TOptional<EBlendFunction> Function = GetBlendFunctionType(InFunction, Functions);
+	const TOptional<EBlendFunction> Function = GetBlendFunctionType(InFunction);
 	return Function ? TweeningUtilsEditor::GetUntintedIconForTweenFunction(*Function) : nullptr;
 }
 
 FLinearColor FBuiltInTweenFunctionContainer::GetColorForFunction(const TSharedRef<ITweenFunction>& InFunction) const
 {
-	const TOptional<EBlendFunction> Function = GetBlendFunctionType(InFunction, Functions);
+	const TOptional<EBlendFunction> Function = GetBlendFunctionType(InFunction);
 	return Function ? TweeningUtilsEditor::GetTintColorForTweenFunction(*Function) : FLinearColor::White;
 }
 
 FText FBuiltInTweenFunctionContainer::GetLabelForFunction(const TSharedRef<ITweenFunction>& InFunction) const
 {
-	const TOptional<EBlendFunction> Function = GetBlendFunctionType(InFunction, Functions);
+	const TOptional<EBlendFunction> Function = GetBlendFunctionType(InFunction);
 	return Function ? TweeningUtilsEditor::GetLabelForBlendFunction(*Function) : FText::GetEmpty();
 }
 	
 FText FBuiltInTweenFunctionContainer::GetToolTipForFunction(const TSharedRef<ITweenFunction>& InFunction) const
 {
-	const TOptional<EBlendFunction> Function = GetBlendFunctionType(InFunction, Functions);
+	const TOptional<EBlendFunction> Function = GetBlendFunctionType(InFunction);
 	return Function ? TweeningUtilsEditor::GetDescriptionForBlendFunction(*Function) : FText::GetEmpty();
 }
 }
diff --git a/Engine/Plugins/Animation/TweeningUtils/Source/TweeningUtilsEditor/Private/Math/Abstraction/KeyBlendingAbstraction.cpp b/Engine/Plugins/Animation/TweeningUtils/Source/TweeningUtilsEditor/Private/Math/Abstraction/KeyBlendingAbstraction.cpp
--- a/Engine/Plugins/Animation/TweeningUtils/Source/TweeningUtilsEditor/Private/Math/Abstraction/KeyBlendingAbstraction.cpp
+++ b/Engine/Plugins/Animation/TweeningUtils/Source/TweeningUtilsEditor/Private/Math/Abstraction/KeyBlendingAbstraction.cpp
@@ -68,10 +68,10 @@ static FStaticBlendFunctionData& GetFunctionData(EBlendFunction InFunction)
 	static FStaticBlendFunctionData Data[Num] =
 	{
 		FStaticBlendFunctionData(
-			EBlendFunction::ControlsToTween, &Blend_ControlsToTween, FTweeningUtilsCommands::Get().SetControlsToTween,
-			LOCTEXT("ControlsToTween.Label", "Tween"),
-			LOCTEXT("ControlsToTween.Abbreviation", "TW"),
-			LOCTEXT("ControlsToTween.Description", "Interpolates between the previous and next keys")
+			EBlendFunction::BlendNeighbor, &Blend_Neighbor, FTweeningUtilsCommands::Get().SetTweenBlendNeighbor,
+			LOCTEXT("BlendNeighbor.Label", "Blend Neighbor"),
+			LOCTEXT("BlendNeighbour.Abbreviation", "BN"),
+			LOCTEXT("BlendNeighbour.Description", "Blend to the next or previous values for selected keys.")
 		),
 		FStaticBlendFunctionData(
 			EBlendFunction::PushPull, &Blend_PushPull, FTweeningUtilsCommands::Get().SetTweenPushPull,
@@ -80,10 +80,16 @@ static FStaticBlendFunctionData& GetFunctionData(EBlendFunction InFunction)
 			LOCTEXT("PushPull.Description", "Push or pull the values to the interpolation between the previous and next keys")
 		),
 		FStaticBlendFunctionData(
-			EBlendFunction::BlendNeighbor, &Blend_Neighbor, FTweeningUtilsCommands::Get().SetTweenBlendNeighbor,
-			LOCTEXT("BlendNeighbor.Label", "Blend Neighbor"),
-			LOCTEXT("BlendNeighbour.Abbreviation", "BN"),
-			LOCTEXT("BlendNeighbour.Description", "Blend to the next or previous values for selected keys.")
+			EBlendFunction::BlendEase, &Blend_Ease, FTweeningUtilsCommands::Get().SetTweenBlendEase,
+			LOCTEXT("BlendEase.Label", "Blend Ease"),
+			LOCTEXT("BlendEase.Abbreviation", "BE"),
+			LOCTEXT("BlendEase.Description", "Blend with an ease falloff to the next or previous value for selected keys")
+		),
+		FStaticBlendFunctionData(
+			EBlendFunction::ControlsToTween, &Blend_ControlsToTween, FTweeningUtilsCommands::Get().SetControlsToTween,
+			LOCTEXT("ControlsToTween.Label", "Tween"),
+			LOCTEXT("ControlsToTween.Abbreviation", "TW"),
+			LOCTEXT("ControlsToTween.Description", "Interpolates between the previous and next keys")
 		),
 		FStaticBlendFunctionData(
 			EBlendFunction::BlendRelative, &Blend_Relative, FTweeningUtilsCommands::Get().SetTweenBlendRelative,
@@ -91,12 +97,6 @@ static FStaticBlendFunctionData& GetFunctionData(EBlendFunction InFunction)
 			LOCTEXT("BlendRelative.Abbreviation", "BR"),
 			LOCTEXT("BlendRelative.Description", "Blend relative to the next or previous value for selected keys")
 		),
-		FStaticBlendFunctionData(
-			EBlendFunction::BlendEase, &Blend_Ease, FTweeningUtilsCommands::Get().SetTweenBlendEase,
-			LOCTEXT("BlendEase.Label", "Blend Ease"),
-			LOCTEXT("BlendEase.Abbreviation", "BE"),
-			LOCTEXT("BlendEase.Description", "Blend with an ease falloff to the next or previous value for selected keys")
-		),
 		FStaticBlendFunctionData(
 			EBlendFunction::SmoothRough, &Blend_SmoothRough, FTweeningUtilsCommands::Get().SetTweenSmoothRough,
 			LOCTEXT("SmoothRough.Label", "Smooth Rough"),
diff --git a/Engine/Plugins/Animation/TweeningUtils/Source/TweeningUtilsEditor/Public/Math/Abstraction/BuiltInTweenFunctionContainer.h b/Engine/Plugins/Animation/TweeningUtils/Source/TweeningUtilsEditor/Public/Math/Abstraction/BuiltInTweenFunctionContainer.h
--- a/Engine/Plugins/Animation/TweeningUtils/Source/TweeningUtilsEditor/Public/Math/Abstraction/BuiltInTweenFunctionContainer.h
+++ b/Engine/Plugins/Animation/TweeningUtils/Source/TweeningUtilsEditor/Public/Math/Abstraction/BuiltInTweenFunctionContainer.h
@@ -9,6 +9,7 @@
 namespace UE::TweeningUtilsEditor
 {
 class FBuiltInTweenFunction;
+enum class EBlendFunction : uint8;
 
 /** Contains all the functions from this module, i.e. those specified by EBlendFunction.  */
 class TWEENINGUTILSEDITOR_API FBuiltInTweenFunctionContainer : public ITweenFunctionContainer
@@ -17,6 +18,9 @@ public:
 
 	FBuiltInTweenFunctionContainer();
 
+	/** @return Gets the type of blend function the tween function corresponds to, or unset if you pass in a function not owned by this container. */
+	TOptional<EBlendFunction> GetBlendFunctionType(const TSharedRef<ITweenFunction>& InFunctionToFind) const;
+
 	//~ Begin ITweenFunctionContainer Interface
 	virtual void ForEachFunction(const TFunctionRef<void(const TSharedRef<ITweenFunction>&)> InConsumer) const override;
 	virtual TSharedPtr<ITweenFunction> GetFunction(int32 InIndex) const override;
diff --git a/Engine/Plugins/Animation/TweeningUtils/Source/TweeningUtilsEditor/Public/Math/Abstraction/KeyBlendingAbstraction.h b/Engine/Plugins/Animation/TweeningUtils/Source/TweeningUtilsEditor/Public/Math/Abstraction/KeyBlendingAbstraction.h
--- a/Engine/Plugins/Animation/TweeningUtils/Source/TweeningUtilsEditor/Public/Math/Abstraction/KeyBlendingAbstraction.h
+++ b/Engine/Plugins/Animation/TweeningUtils/Source/TweeningUtilsEditor/Public/Math/Abstraction/KeyBlendingAbstraction.h
@@ -20,11 +20,11 @@ struct FKeyBlendInfo;
 /** Util enum which is useful for UI code to abstract functions.*/
 enum class EBlendFunction : uint8
 {
-	ControlsToTween,
-	PushPull,
 	BlendNeighbor,
-	BlendRelative,
+	PushPull,
 	BlendEase,
+	ControlsToTween,
+	BlendRelative,
 	SmoothRough,
 
 	// Steps for adding a new function:
