Set missing default value for FTweenPointStyle::HitTestSize.

#rnx

diff --git a/Engine/Plugins/Animation/TweeningUtils/Source/TweeningUtilsEditor/Public/Widgets/TweenSliderStyle.h b/Engine/Plugins/Animation/TweeningUtils/Source/TweeningUtilsEditor/Public/Widgets/TweenSliderStyle.h
--- a/Engine/Plugins/Animation/TweeningUtils/Source/TweeningUtilsEditor/Public/Widgets/TweenSliderStyle.h
+++ b/Engine/Plugins/Animation/TweeningUtils/Source/TweeningUtilsEditor/Public/Widgets/TweenSliderStyle.h
@@ -1,4 +1,4 @@
-﻿// Copyright Epic Games, Inc. All Rights Reserved.
+// Copyright Epic Games, Inc. All Rights Reserved.
 
 #pragma once
 
@@ -40,7 +40,7 @@ struct TWEENINGUTILSEDITOR_API FTweenPointStyle
 	
 	/** The size of the hit box against which hit tests are made for this point (so the user does not need to click them pixel perfectly). */
 	UPROPERTY(EditAnywhere, Category=Appearance)
-	FVector2D HitTestSize;
+	FVector2D HitTestSize = { 10.0, 12.0 };
 	FTweenPointStyle& SetHitTestSize( const FVector2D& InSize ){ HitTestSize = InSize; return *this; }
 	
 	void GetResources(TArray<const FSlateBrush*>& OutBrushes) const;
