Curve Editor: Tween tool now inserts keys at the scrub position (if it has one)

diff --git a/Engine/Plugins/Animation/TweeningUtils/Source/TweeningUtilsEditor/Private/Math/ContiguousKeyMapping.cpp b/Engine/Plugins/Animation/TweeningUtils/Source/TweeningUtilsEditor/Private/Math/ContiguousKeyMapping.cpp
--- a/Engine/Plugins/Animation/TweeningUtils/Source/TweeningUtilsEditor/Private/Math/ContiguousKeyMapping.cpp
+++ b/Engine/Plugins/Animation/TweeningUtils/Source/TweeningUtilsEditor/Private/Math/ContiguousKeyMapping.cpp
@@ -35,75 +35,79 @@ static int32 GetIndex(const TArray<FKeyPosition>& Keys, double Time)
 	return Index;
 }
 
-static TMap<FCurveModelID, FContiguousKeyMapping::FContiguousKeysArray> ComputeKeyMap(const FCurveEditor& InCurveEditor)
+static void AppendKeyArray(
+	const FCurveEditor& InCurveEditor,
+	const FCurveModelID& InCurveId, TConstArrayView<FKeyHandle> KeysToBlend,
+	TMap<FCurveModelID, FContiguousKeyMapping::FContiguousKeysArray>& OutKeyMap
+	)
 {
-	using FContiguousKeysArray = FContiguousKeyMapping::FContiguousKeysArray;
-	
-	TArray<FKeyHandle> AllKeyHandles;
-	TArray<FKeyPosition> AllKeyPositions;
-
-	TArray<FKeyHandle> SelectedKeyHandles;
-	TArray<FKeyPosition> SelectedKeyPositions;
-
-	const TMap<FCurveModelID, FKeyHandleSet>& SelectionKeyMap = InCurveEditor.Selection.GetAll();
-
-	TArray<int32> ContiguousKeyIndices;
-
-	TMap<FCurveModelID, FContiguousKeyMapping::FContiguousKeysArray> KeyMap;
-	for (const TPair<FCurveModelID, FKeyHandleSet>& Pair : SelectionKeyMap)
+	if (const FCurveModel* Curve = InCurveEditor.FindCurve(InCurveId))
 	{
-		if (const FCurveModel* Curve = InCurveEditor.FindCurve(Pair.Key))
+		TArray<FKeyHandle> AllKeyHandles;
+		Curve->GetKeys(
+			TNumericLimits<double>::Lowest(), TNumericLimits<double>::Max(), TNumericLimits<double>::Lowest(), TNumericLimits<double>::Max(),
+			AllKeyHandles
+			);
+		
+		TArray<FKeyPosition> AllKeyPositions;
+		AllKeyPositions.SetNum(AllKeyHandles.Num());
+		Curve->GetKeyPositions(AllKeyHandles, AllKeyPositions);
+
+		// Get all the selected keys
+		TArray<FKeyPosition> KeyPositions;
+		KeyPositions.SetNum(KeysToBlend.Num());
+		Curve->GetKeyPositions(KeysToBlend, KeyPositions);
+
+		FContiguousKeyMapping::FContiguousKeysArray& KeyArray = OutKeyMap.Add(
+			InCurveId, FContiguousKeyMapping::FContiguousKeysArray(AllKeyHandles, AllKeyPositions)
+			);
+
+		TArray<int32> SelectedIndices;
+		for (int32 Index = 0; Index < KeyPositions.Num(); ++Index)
 		{
-			AllKeyHandles.Reset();
-			Curve->GetKeys(TNumericLimits<double>::Lowest(), TNumericLimits<double>::Max(), TNumericLimits<double>::Lowest(), TNumericLimits<double>::Max(), AllKeyHandles);
-			AllKeyPositions.SetNum(AllKeyHandles.Num());
-			Curve->GetKeyPositions(AllKeyHandles, AllKeyPositions);
-			SelectedKeyHandles.Reset(Pair.Value.Num());
-			SelectedKeyHandles.Append(Pair.Value.AsArray().GetData(), Pair.Value.Num());
-
-			// Get all the selected keys
-			SelectedKeyPositions.SetNum(SelectedKeyHandles.Num());
-			Curve->GetKeyPositions(SelectedKeyHandles, SelectedKeyPositions);
-
-			ContiguousKeyIndices.Reset();
-			FContiguousKeysArray& KeyArray = KeyMap.Add(Pair.Key, FContiguousKeysArray(AllKeyHandles, AllKeyPositions));
-
-			TArray<int32> SelectedIndices;
-			for (int32 Index = 0; Index < SelectedKeyPositions.Num(); ++Index)
-			{
-				int32 AllKeysIndex = GetIndex(AllKeyPositions, SelectedKeyPositions[Index].InputValue);
-				SelectedIndices.Add(AllKeysIndex);
-			}
-			SelectedIndices.Sort();
-			for(int32 AllKeysIndex: SelectedIndices)
+			int32 AllKeysIndex = GetIndex(AllKeyPositions, KeyPositions[Index].InputValue);
+			SelectedIndices.Add(AllKeysIndex);
+		}
+		SelectedIndices.Sort();
+		
+		TArray<int32> ContiguousKeyIndices;
+		for(int32 AllKeysIndex: SelectedIndices)
+		{
+			if (ContiguousKeyIndices.Num() > 0) //see if this key is next to the previous one
 			{
-				if (ContiguousKeyIndices.Num() > 0) //see if this key is next to the previous one
+				if (ContiguousKeyIndices[ContiguousKeyIndices.Num() - 1] + 1 == AllKeysIndex)
 				{
-					if (ContiguousKeyIndices[ContiguousKeyIndices.Num() - 1] + 1 == AllKeysIndex)
-					{
-						ContiguousKeyIndices.Add(AllKeysIndex);
-					}
-					else //not contiguous, so need to add to the OutKeyMap
-					{
-						//now start new set with the new index
-						KeyArray.Add(ContiguousKeyIndices);
-						ContiguousKeyIndices.Reset();
-						ContiguousKeyIndices.Add(AllKeysIndex);
-					}
+					ContiguousKeyIndices.Add(AllKeysIndex);
 				}
-				else//first key in this set so just add
+				else //not contiguous, so need to add to the OutKeyMap
 				{
+					//now start new set with the new index
+					KeyArray.Add(ContiguousKeyIndices);
+					ContiguousKeyIndices.Reset();
 					ContiguousKeyIndices.Add(AllKeysIndex);
 				}
 			}
-			if (ContiguousKeyIndices.Num() > 0)
+			else//first key in this set so just add
 			{
-				KeyArray.Add(ContiguousKeyIndices);
-				ContiguousKeyIndices.Reset();
+				ContiguousKeyIndices.Add(AllKeysIndex);
 			}
 		}
+		if (ContiguousKeyIndices.Num() > 0)
+		{
+			KeyArray.Add(ContiguousKeyIndices);
+			ContiguousKeyIndices.Reset();
+		}
+	}
+}
+
+static TMap<FCurveModelID, FContiguousKeyMapping::FContiguousKeysArray> ComputeKeyMap(const FCurveEditor& InCurveEditor)
+{
+	const TMap<FCurveModelID, FKeyHandleSet>& SelectionKeyMap = InCurveEditor.Selection.GetAll();
+	TMap<FCurveModelID, FContiguousKeyMapping::FContiguousKeysArray> KeyMap;
+	for (const TPair<FCurveModelID, FKeyHandleSet>& Pair : SelectionKeyMap)
+	{
+		AppendKeyArray(InCurveEditor, Pair.Key, Pair.Value.AsArray(), KeyMap);
 	}
-	
 	return KeyMap;
 }
 
@@ -148,4 +152,9 @@ void FContiguousKeyMapping::FContiguousKeysArray::Add(const TArray<int32>& InCon
 FContiguousKeyMapping::FContiguousKeyMapping(const FCurveEditor& InCurveEditor)
 	: KeyMap(ComputeKeyMap(InCurveEditor))
 {}
+
+void FContiguousKeyMapping::Append(const FCurveEditor& InCurveEditor, const FCurveModelID& InCurveId, TConstArrayView<FKeyHandle> InKeysToBlend)
+{
+	AppendKeyArray(InCurveEditor, InCurveId, InKeysToBlend, KeyMap);
+}
 }
diff --git a/Engine/Plugins/Animation/TweeningUtils/Source/TweeningUtilsEditor/Private/Widgets/MVC/Models/CurveTweeningModel.cpp b/Engine/Plugins/Animation/TweeningUtils/Source/TweeningUtilsEditor/Private/Widgets/MVC/Models/CurveTweeningModel.cpp
--- a/Engine/Plugins/Animation/TweeningUtils/Source/TweeningUtilsEditor/Private/Widgets/MVC/Models/CurveTweeningModel.cpp
+++ b/Engine/Plugins/Animation/TweeningUtils/Source/TweeningUtilsEditor/Private/Widgets/MVC/Models/CurveTweeningModel.cpp
@@ -12,6 +12,12 @@ FCurveTweenModel::FCurveTweenModel(TWeakPtr<FCurveEditor> InWeakCurveEditor, con
 	, WeakCurveEditor(InWeakCurveEditor)
 {}
 
+bool FCurveTweenModel::HasAnythingToBlend() const
+{
+	const TSharedPtr<FCurveEditor> CurveEditor = WeakCurveEditor.Pin();
+	return CurveEditor && !CurveEditor->GetSelection().IsEmpty();
+}
+
 void FCurveTweenModel::StartBlendOperation()
 {
 	const TSharedPtr<FCurveEditor> CurveEditor = WeakCurveEditor.Pin();
diff --git a/Engine/Plugins/Animation/TweeningUtils/Source/TweeningUtilsEditor/Public/Math/ContiguousKeyMapping.h b/Engine/Plugins/Animation/TweeningUtils/Source/TweeningUtilsEditor/Public/Math/ContiguousKeyMapping.h
--- a/Engine/Plugins/Animation/TweeningUtils/Source/TweeningUtilsEditor/Public/Math/ContiguousKeyMapping.h
+++ b/Engine/Plugins/Animation/TweeningUtils/Source/TweeningUtilsEditor/Public/Math/ContiguousKeyMapping.h
@@ -41,6 +41,9 @@ struct TWEENINGUTILSEDITOR_API FContiguousKeyMapping
 	FContiguousKeyMapping() = default;
 	explicit FContiguousKeyMapping(const FCurveEditor& InCurveEditor);
 
+	/** Appends keys from the given curve.*/
+	void Append(const FCurveEditor& InCurveEditor, const FCurveModelID& InCurveId, TConstArrayView<FKeyHandle> InKeysToBlend);
+
 	TMap<FCurveModelID, FContiguousKeysArray> KeyMap;
 };
 }
diff --git a/Engine/Plugins/Animation/TweeningUtils/Source/TweeningUtilsEditor/Public/Widgets/MVC/Models/CurveTweenModel.h b/Engine/Plugins/Animation/TweeningUtils/Source/TweeningUtilsEditor/Public/Widgets/MVC/Models/CurveTweenModel.h
--- a/Engine/Plugins/Animation/TweeningUtils/Source/TweeningUtilsEditor/Public/Widgets/MVC/Models/CurveTweenModel.h
+++ b/Engine/Plugins/Animation/TweeningUtils/Source/TweeningUtilsEditor/Public/Widgets/MVC/Models/CurveTweenModel.h
@@ -16,6 +16,9 @@ class TWEENINGUTILSEDITOR_API FCurveTweenModel : public FFunctionBasedTweenModel
 public:
 
 	explicit FCurveTweenModel(TWeakPtr<FCurveEditor> InWeakCurveEditor, const TSharedRef<ITweenFunction>& InTweenFunction);
+
+	/** @return Whether it makes sense to call StartBlendOperation, etc. */
+	bool HasAnythingToBlend() const;
 	
 	//~ Begin FTweenModel Interface
 	virtual void StartBlendOperation() override;
@@ -23,11 +26,11 @@ public:
 	virtual void BlendValues(float InNormalizedValue) override;
 	//~ Begin FTweenModel Interface
 
-private:
-
+protected:
+	
 	/** The curve editor on which to tween the curves. */
 	const TWeakPtr<FCurveEditor> WeakCurveEditor;
-	
+
 	/** Created in StartBlendOperation and used for the entirety of the blend operation. */
 	FContiguousKeyMapping ContiguousKeySelection;
 };
diff --git a/Engine/Plugins/Animation/TweeningUtils/Source/TweeningUtilsEditor/Public/Widgets/MVC/Models/FunctionBasedTweenModel.h b/Engine/Plugins/Animation/TweeningUtils/Source/TweeningUtilsEditor/Public/Widgets/MVC/Models/FunctionBasedTweenModel.h
--- a/Engine/Plugins/Animation/TweeningUtils/Source/TweeningUtilsEditor/Public/Widgets/MVC/Models/FunctionBasedTweenModel.h
+++ b/Engine/Plugins/Animation/TweeningUtils/Source/TweeningUtilsEditor/Public/Widgets/MVC/Models/FunctionBasedTweenModel.h
@@ -6,23 +6,20 @@
 #include "Templates/SharedPointer.h"
 #include "TweenModel.h"
 
-namespace UE::TweeningUtilsEditor
-{
-	class ITweenFunction;
-}
-
 class FCurveEditor;
 
 namespace UE::TweeningUtilsEditor
 {
+class ITweenFunction;
+	
 /** Base class for models that want to implement BlendValues based off of some injectable ITweenFunction. */
 class FFunctionBasedTweenModel : public FTweenModel
 {
 public:
 
-	explicit FFunctionBasedTweenModel(const TSharedRef<ITweenFunction>& InTweenFunction) : TweenFunction(InTweenFunction) {}
+	explicit FFunctionBasedTweenModel(const TSharedRef<ITweenFunction>& InInitialFunction) : TweenFunction(InInitialFunction) {}
 
-	void SetTweenFunction(const TSharedRef<ITweenFunction>& InTweenFunction) { TweenFunction = InTweenFunction; }
+	virtual void SetTweenFunction(const TSharedRef<ITweenFunction>& InTweenFunction) { TweenFunction = InTweenFunction; }
 	const TSharedRef<ITweenFunction>& GetTweenFunction() const { return TweenFunction; }
 	
 private:
diff --git a/Engine/Plugins/Editor/CurveEditorTools/Source/CurveEditorTools/Private/CurveEditorToolsModule.cpp b/Engine/Plugins/Editor/CurveEditorTools/Source/CurveEditorTools/Private/CurveEditorToolsModule.cpp
--- a/Engine/Plugins/Editor/CurveEditorTools/Source/CurveEditorTools/Private/CurveEditorToolsModule.cpp
+++ b/Engine/Plugins/Editor/CurveEditorTools/Source/CurveEditorTools/Private/CurveEditorToolsModule.cpp
@@ -4,7 +4,7 @@
 #include "Framework/Commands/UICommandList.h"
 #include "ICurveEditorModule.h"
 #include "Extensions/CurveEditorFocusExtension.h"
-#include "Extensions/TweenEditorExtension.h"
+#include "Extensions/Tweening/TweenEditorExtension.h"
 #include "CurveEditorToolCommands.h"
 #include "Tools/CurveEditorTransformTool.h"
 #include "Tools/CurveEditorRetimeTool.h"
@@ -122,7 +122,7 @@ TSharedRef<ICurveEditorExtension> FCurveEditorToolsModule::CreateFocusExtension(
 
 TSharedRef<ICurveEditorExtension> FCurveEditorToolsModule::CreateTweenExtension(TWeakPtr<FCurveEditor> InCurveEditor)
 {
-	return MakeShared<UE::CurveEditor::FTweenEditorExtension>(InCurveEditor);
+	return MakeShared<UE::CurveEditorTools::FTweenEditorExtension>(InCurveEditor);
 }
 
 TUniquePtr<ICurveEditorToolExtension> FCurveEditorToolsModule::CreateTransformToolExtension(TWeakPtr<FCurveEditor> InCurveEditor)
diff --git a/Engine/Plugins/Editor/CurveEditorTools/Source/CurveEditorTools/Private/Extensions/TweenEditorExtension.cpp b/Engine/Plugins/Editor/CurveEditorTools/Source/CurveEditorTools/Private/Extensions/TweenEditorExtension.cpp
deleted file mode 100644
diff --git a/Engine/Plugins/Editor/CurveEditorTools/Source/CurveEditorTools/Private/Extensions/TweenEditorExtension.h b/Engine/Plugins/Editor/CurveEditorTools/Source/CurveEditorTools/Private/Extensions/TweenEditorExtension.h
deleted file mode 100644
diff --git a/Engine/Plugins/Editor/CurveEditorTools/Source/CurveEditorTools/Private/Extensions/Tweening/ContextAwareTweeningModel.cpp b/Engine/Plugins/Editor/CurveEditorTools/Source/CurveEditorTools/Private/Extensions/Tweening/ContextAwareTweeningModel.cpp
new file mode 100644
--- /dev/null
+++ b/Engine/Plugins/Editor/CurveEditorTools/Source/CurveEditorTools/Private/Extensions/Tweening/ContextAwareTweeningModel.cpp
@@ -0,0 +1,75 @@
+﻿// Copyright Epic Games, Inc. All Rights Reserved.
+
+#include "ContextAwareTweeningModel.h"
+
+#include "KeyInsertingTweenModel.h"
+
+namespace UE::CurveEditorTools
+{
+FContextAwareTweeningModel::FContextAwareTweeningModel(
+	TWeakPtr<FCurveEditor> InWeakCurveEditor,
+	const TSharedRef<TweeningUtilsEditor::ITweenFunction>& InTweenFunction
+	)
+	: FFunctionBasedTweenModel(InTweenFunction)
+	, CurveEditor(InWeakCurveEditor)
+	, SelectionBasedModel(InWeakCurveEditor, InTweenFunction)
+	, NoKeySelectionModel(MakeShared<FKeyInsertingTweenModel>(InWeakCurveEditor, InTweenFunction))
+{}
+
+void FContextAwareTweeningModel::SetScaleMode(TweeningUtilsEditor::ETweenScaleMode InMode)
+{
+	Super::SetScaleMode(InMode);
+	SelectionBasedModel.SetScaleMode(InMode);
+	if (NoKeySelectionModel)
+	{
+		NoKeySelectionModel->SetScaleMode(InMode);
+	}
+}
+
+void FContextAwareTweeningModel::SetTweenFunction(const TSharedRef<TweeningUtilsEditor::ITweenFunction>& InTweenFunction)
+{
+	Super::SetTweenFunction(InTweenFunction);
+	SelectionBasedModel.SetTweenFunction(InTweenFunction);
+	if (NoKeySelectionModel)
+	{
+		NoKeySelectionModel->SetTweenFunction(InTweenFunction);
+	}
+}
+
+void FContextAwareTweeningModel::StartBlendOperation()
+{
+	AccessCurrentModel([this](FFunctionBasedTweenModel& Model)
+	{
+		Model.StartBlendOperation();
+	});
+}
+
+void FContextAwareTweeningModel::StopBlendOperation()
+{
+	AccessCurrentModel([](FFunctionBasedTweenModel& Model)
+	{
+		Model.StopBlendOperation();
+	});
+}
+
+void FContextAwareTweeningModel::BlendValues(float InNormalizedValue)
+{
+	AccessCurrentModel([InNormalizedValue](FFunctionBasedTweenModel& Model)
+	{
+		Model.BlendValues(InNormalizedValue);
+	});
+}
+
+template <typename TCallback> requires std::is_invocable_v<TCallback, TweeningUtilsEditor::FFunctionBasedTweenModel&>
+void FContextAwareTweeningModel::AccessCurrentModel(TCallback&& InCallback)
+{
+	if (SelectionBasedModel.HasAnythingToBlend())
+	{
+		InCallback(SelectionBasedModel);
+	}
+	else if (NoKeySelectionModel)
+	{
+		InCallback(*NoKeySelectionModel.Get());
+	}
+}
+}
diff --git a/Engine/Plugins/Editor/CurveEditorTools/Source/CurveEditorTools/Private/Extensions/Tweening/ContextAwareTweeningModel.h b/Engine/Plugins/Editor/CurveEditorTools/Source/CurveEditorTools/Private/Extensions/Tweening/ContextAwareTweeningModel.h
new file mode 100644
--- /dev/null
+++ b/Engine/Plugins/Editor/CurveEditorTools/Source/CurveEditorTools/Private/Extensions/Tweening/ContextAwareTweeningModel.h
@@ -0,0 +1,58 @@
+﻿// Copyright Epic Games, Inc. All Rights Reserved.
+
+#pragma once
+
+#include "Containers/ContainersFwd.h"
+#include "Widgets/MVC/Models/CurveTweenModel.h"
+#include "Widgets/MVC/Models/FunctionBasedTweenModel.h"
+
+#include <type_traits>
+
+class FCurveEditor;
+template<typename ObjectType> class TAttribute;
+
+namespace UE::CurveEditorTools
+{
+/**
+ * This model tweens keys depending on the user selection:
+ * - When keys are selected, FCurveTweenModel will tween the selected keys.
+ * - When no keys are selected, we'll insert keys at the time controller's scrub position to all visible curves (those selected in tree view).
+ *   If there's no such controller, nothing happens.
+ */
+class FContextAwareTweeningModel : public TweeningUtilsEditor::FFunctionBasedTweenModel
+{
+	using Super = FFunctionBasedTweenModel;
+public:
+	
+	/**
+	 * @param InWeakCurveEditor Used to get key info
+	 * @param InInitialTweenFunction The initial function to blend with
+	 */
+	explicit FContextAwareTweeningModel(
+		TWeakPtr<FCurveEditor> InWeakCurveEditor,
+		const TSharedRef<TweeningUtilsEditor::ITweenFunction>& InInitialTweenFunction
+		);
+
+	//~ Begin FTweenModel Interface
+	virtual void SetScaleMode(TweeningUtilsEditor::ETweenScaleMode InMode) override;
+	virtual void SetTweenFunction(const TSharedRef<TweeningUtilsEditor::ITweenFunction>& InTweenFunction) override;
+	virtual void StartBlendOperation() override;
+	virtual void StopBlendOperation() override;
+	virtual void BlendValues(float InNormalizedValue) override;
+	//~ End FTweenModel Interface
+
+private:
+
+	/** The curve editor of which the keys are tweened. */
+	const TWeakPtr<FCurveEditor> CurveEditor;
+
+	/** Tweens when there is a key selection. */
+	TweeningUtilsEditor::FCurveTweenModel SelectionBasedModel;
+	/** Injected model that tweens when there is no key selection. Can be nullptr. */
+	const TSharedPtr<FFunctionBasedTweenModel> NoKeySelectionModel;
+
+	template<typename TCallback> requires std::is_invocable_v<TCallback, TweeningUtilsEditor::FFunctionBasedTweenModel&>
+	void AccessCurrentModel(TCallback&& InCallback);
+};
+}
+
diff --git a/Engine/Plugins/Editor/CurveEditorTools/Source/CurveEditorTools/Private/Extensions/Tweening/KeyInsertingTweenModel.cpp b/Engine/Plugins/Editor/CurveEditorTools/Source/CurveEditorTools/Private/Extensions/Tweening/KeyInsertingTweenModel.cpp
new file mode 100644
--- /dev/null
+++ b/Engine/Plugins/Editor/CurveEditorTools/Source/CurveEditorTools/Private/Extensions/Tweening/KeyInsertingTweenModel.cpp
@@ -0,0 +1,72 @@
+﻿// Copyright Epic Games, Inc. All Rights Reserved.
+
+#include "KeyInsertingTweenModel.h"
+
+#include "CurveEditor.h"
+#include "Math/CurveBlending.h"
+#include "Math/Abstraction/ITweenFunction.h"
+
+namespace UE::CurveEditorTools
+{
+FKeyInsertingTweenModel::FKeyInsertingTweenModel(
+	TWeakPtr<FCurveEditor> InWeakCurveEditor, const TSharedRef<TweeningUtilsEditor::ITweenFunction>& InInitialFunction
+	)
+	: FCurveTweenModel(MoveTemp(InWeakCurveEditor), InInitialFunction)
+{}
+
+void FKeyInsertingTweenModel::StartBlendOperation()
+{
+	const TSharedPtr<FCurveEditor> CurveEditor = WeakCurveEditor.Pin();
+	if (!ensure(CurveEditor))
+	{
+		return;
+	}
+
+	// Not all curve editors have a time slider controller
+	const TSharedPtr<ITimeSliderController> TimeSliderController = CurveEditor->GetTimeSliderController();
+	if (!TimeSliderController)
+	{
+		return;
+	}
+	
+	const double ScrubTime = TimeSliderController->GetTickResolution().AsSeconds(TimeSliderController->GetScrubPosition());
+	for (const TTuple<FCurveEditorTreeItemID, ECurveEditorTreeSelectionState>& Pair : CurveEditor->GetTreeSelection())
+	{
+		if (Pair.Value == ECurveEditorTreeSelectionState::None)
+		{
+			continue;
+		}
+		
+		const FCurveEditorTreeItem& TreeItem = CurveEditor->GetTreeItem(Pair.Key);
+		for (const FCurveModelID& CurveModelID : TreeItem.GetCurves())
+		{
+			FCurveModel* CurveModel = CurveEditor->FindCurve(CurveModelID);
+			if (!CurveModel)
+			{
+				return;
+			}
+
+			double ValueAtScrubTime;
+			if (!CurveModel->Evaluate(ScrubTime, ValueAtScrubTime))
+			{
+				continue;
+			}
+
+			// Try to insert the key with similar interpolation modes as its neighbours
+			FKeyAttributes KeyAttributes;
+			const TPair<ERichCurveInterpMode, ERichCurveTangentMode> Modes = CurveModel->GetInterpolationMode(ScrubTime, RCIM_Cubic, RCTM_SmartAuto);
+			KeyAttributes.SetInterpMode(Modes.Key);
+			KeyAttributes.SetTangentMode(Modes.Value);
+
+			const TOptional<FKeyHandle> KeyHandle = CurveModel->AddKey({ ScrubTime, ValueAtScrubTime }, KeyAttributes);
+			if (!KeyHandle)
+			{
+				continue;
+			}
+
+			// This will make Blend function blend the key we just inserted
+			ContiguousKeySelection.Append(*CurveEditor, CurveModelID, { *KeyHandle });
+		}
+	}
+}
+}
diff --git a/Engine/Plugins/Editor/CurveEditorTools/Source/CurveEditorTools/Private/Extensions/Tweening/KeyInsertingTweenModel.h b/Engine/Plugins/Editor/CurveEditorTools/Source/CurveEditorTools/Private/Extensions/Tweening/KeyInsertingTweenModel.h
new file mode 100644
--- /dev/null
+++ b/Engine/Plugins/Editor/CurveEditorTools/Source/CurveEditorTools/Private/Extensions/Tweening/KeyInsertingTweenModel.h
@@ -0,0 +1,27 @@
+﻿// Copyright Epic Games, Inc. All Rights Reserved.
+
+#pragma once
+
+#include "Templates/SharedPointer.h"
+#include "Widgets/MVC/Models/CurveTweenModel.h"
+
+class FCurveEditor;
+
+namespace UE::CurveEditorTools
+{
+/** Inserts keys of visible curves into the current scrubber position. */
+class FKeyInsertingTweenModel : public TweeningUtilsEditor::FCurveTweenModel
+{
+public:
+	
+	explicit FKeyInsertingTweenModel(
+		TWeakPtr<FCurveEditor> InWeakCurveEditor,
+		const TSharedRef<TweeningUtilsEditor::ITweenFunction>& InInitialFunction
+		);
+
+	//~ Begin FTweenModel
+	virtual void StartBlendOperation() override;
+	//~ Begin FTweenModel
+};
+}
+
diff --git a/Engine/Plugins/Editor/CurveEditorTools/Source/CurveEditorTools/Private/Extensions/Tweening/TweenEditorExtension.cpp b/Engine/Plugins/Editor/CurveEditorTools/Source/CurveEditorTools/Private/Extensions/Tweening/TweenEditorExtension.cpp
new file mode 100644
--- /dev/null
+++ b/Engine/Plugins/Editor/CurveEditorTools/Source/CurveEditorTools/Private/Extensions/Tweening/TweenEditorExtension.cpp
@@ -0,0 +1,86 @@
+﻿// Copyright Epic Games, Inc. All Rights Reserved.
+
+#include "TweenEditorExtension.h"
+
+#include "ContextAwareTweeningModel.h"
+#include "CurveEditor.h"
+#include "Framework/MultiBox/MultiBoxBuilder.h"
+#include "Framework/MultiBox/MultiBoxExtender.h"
+#include "Math/Abstraction/BuiltInTweenFunctionContainer.h"
+#include "Misc/Optional.h"
+#include "Widgets/MVC/Models/TransactionalTweenModelProxy.h"
+#include "Widgets/MVC/TweenToolbarController.h"
+
+namespace UE::CurveEditorTools
+{
+FTweenEditorExtension::FTweenEditorExtension(TWeakPtr<FCurveEditor> InCurveEditor)
+	: WeakCurveEditor(InCurveEditor)
+	, BlendFunctionContainer(MakeShared<TweeningUtilsEditor::FBuiltInTweenFunctionContainer>())
+{}
+
+void FTweenEditorExtension::BindCommands(TSharedRef<FUICommandList> InCommandList)
+{
+	InitControllers(InCommandList);
+}
+
+TSharedPtr<FExtender> FTweenEditorExtension::MakeToolbarExtender(const TSharedRef<FUICommandList>& InCommandList)
+{
+	using namespace TweeningUtilsEditor;
+	TSharedRef<FExtender> Extender = MakeShared<FExtender>();
+	
+	Extender->AddToolBarExtension("Tangents", EExtensionHook::After, InCommandList, FToolBarExtensionDelegate::CreateLambda(
+		[this, WeakCommandList = InCommandList.ToWeakPtr()](FToolBarBuilder& ToolbarBuilder)
+		{
+			const TSharedPtr<FUICommandList> CommandListPin = WeakCommandList.Pin();
+			if (ensure(CommandListPin) && InitControllers(CommandListPin.ToSharedRef()))
+			{
+				// Visually move the slider widget when the user uses the U+LMB button command to indirectly move the mouse
+				TAttribute<TOptional<float>> OverrideSliderPosition =
+					TAttribute<TOptional<float>>::CreateLambda([this]()
+					{
+						return TweenControllers->MouseSlidingController.GetCurrentSliderPosition();
+					});
+				
+				ToolbarBuilder.BeginStyleOverride(TEXT("CurveEditorTweenToolbar")); // Forces min sizes for combo button
+				ToolbarBuilder.BeginSection("Tween");
+				TweenControllers->ToolbarController.AddToToolbar(ToolbarBuilder, MoveTemp(OverrideSliderPosition));
+				ToolbarBuilder.EndSection();
+				ToolbarBuilder.EndStyleOverride();
+			}
+		}));
+	
+	return Extender;
+}
+	
+FTweenEditorExtension::FTweenControllers::FTweenControllers(
+	const TSharedRef<TweeningUtilsEditor::FFunctionBasedTweenModel>& InTweenModel,
+	const TSharedRef<FUICommandList>& InCommandList,
+	const TSharedRef<TweeningUtilsEditor::ITweenFunctionContainer>& InTweenFunctions
+	)
+	: ToolbarController(InTweenModel, InCommandList, InTweenFunctions)
+	, CycleFunctionController(InTweenModel, InTweenFunctions, InCommandList)
+	, MouseSlidingController(InTweenModel, InCommandList)
+{}
+
+bool FTweenEditorExtension::InitControllers(TSharedRef<FUICommandList> InCommandList)
+{
+	if (TweenControllers)
+	{
+		return true;
+	}
+
+	if (const TSharedPtr<FCurveEditor> CurveEditorPin = WeakCurveEditor.Pin(); ensure(CurveEditorPin))
+	{
+		using namespace TweeningUtilsEditor;
+		using FModel = FContextAwareTweeningModel;
+
+		const TSharedRef<ITweenFunction> InitialFunction = BlendFunctionContainer->GetFunction(0).ToSharedRef();
+		const TSharedRef<TTransactionalTweenModelProxy<FModel>> TweenModel = MakeShared<TTransactionalTweenModelProxy<FModel>>(
+			WeakCurveEditor, InitialFunction
+			);
+		TweenControllers.Emplace(TweenModel, InCommandList, BlendFunctionContainer);
+		return true;
+	}
+	return false;
+}
+}
diff --git a/Engine/Plugins/Editor/CurveEditorTools/Source/CurveEditorTools/Private/Extensions/Tweening/TweenEditorExtension.h b/Engine/Plugins/Editor/CurveEditorTools/Source/CurveEditorTools/Private/Extensions/Tweening/TweenEditorExtension.h
new file mode 100644
--- /dev/null
+++ b/Engine/Plugins/Editor/CurveEditorTools/Source/CurveEditorTools/Private/Extensions/Tweening/TweenEditorExtension.h
@@ -0,0 +1,64 @@
+﻿// Copyright Epic Games, Inc. All Rights Reserved.
+
+#pragma once
+
+#include "ICurveEditorExtension.h"
+#include "Misc/Optional.h"
+#include "Templates/SharedPointer.h"
+#include "Widgets/MVC/CycleFunctionController.h"
+#include "Widgets/MVC/TweenMouseSlidingController.h"
+#include "Widgets/MVC/TweenToolbarController.h"
+
+class FCurveEditor;
+class FExtender;
+namespace UE::TweeningUtilsEditor { class FBuiltInTweenFunctionContainer; }
+namespace UE::TweeningUtilsEditor { class FFunctionBasedTweenModel; }
+namespace UE::TweeningUtilsEditor { class ITweenFunctionContainer; }
+
+namespace UE::CurveEditorTools
+{
+/** Manages the tweening tool in the curve editor. */
+class FTweenEditorExtension : public ICurveEditorExtension, public FNoncopyable
+{
+public:
+
+	explicit FTweenEditorExtension(TWeakPtr<FCurveEditor> InCurveEditor);
+
+	//~ Begin ICurveEditorExtension Interface
+	virtual void BindCommands(TSharedRef<FUICommandList> InCommandList) override;
+	virtual TSharedPtr<FExtender> MakeToolbarExtender(const TSharedRef<FUICommandList>& InCommandList) override;
+	//~ End ICurveEditorExtension Interface
+
+private:
+
+	/** The editor that owns us. */
+	const TWeakPtr<FCurveEditor> WeakCurveEditor;
+
+	/** Holds the functions that ToolbarController shows: just the default, built-in ones.*/
+	const TSharedRef<TweeningUtilsEditor::FBuiltInTweenFunctionContainer> BlendFunctionContainer;
+
+	struct FTweenControllers
+	{
+		/** Manages the toolbar widget */
+		const TweeningUtilsEditor::FTweenToolbarController ToolbarController;
+
+		/** Cycles functions (Shift + U). */
+		const TweeningUtilsEditor::FCycleFunctionController CycleFunctionController;
+
+		/** Allows indirect movement of the slider by using u + Move Mouse */
+		const TweeningUtilsEditor::FTweenMouseSlidingController MouseSlidingController;
+		
+		explicit FTweenControllers(
+			const TSharedRef<TweeningUtilsEditor::FFunctionBasedTweenModel>& InTweenModel,
+			const TSharedRef<FUICommandList>& InCommandList,
+			const TSharedRef<TweeningUtilsEditor::ITweenFunctionContainer>& InTweenFunctions
+			);
+	};
+	/** Created by InitControllers. */
+	TOptional<FTweenControllers> TweenControllers;
+
+	/** Makes sure that ToolbarController is created. */
+	bool InitControllers(TSharedRef<FUICommandList> InCommandList);
+};
+}
+
