Curve Editor: Last selected tween function is now saved as user preference. E.g. after editor restarts, the preferred function is selected in the tool by default.

diff --git a/Engine/Plugins/Animation/TweeningUtils/Source/TweeningUtilsEditor/Private/Math/Abstraction/BuiltInTweenFunctionContainer.cpp b/Engine/Plugins/Animation/TweeningUtils/Source/TweeningUtilsEditor/Private/Math/Abstraction/BuiltInTweenFunctionContainer.cpp
--- a/Engine/Plugins/Animation/TweeningUtils/Source/TweeningUtilsEditor/Private/Math/Abstraction/BuiltInTweenFunctionContainer.cpp
+++ b/Engine/Plugins/Animation/TweeningUtils/Source/TweeningUtilsEditor/Private/Math/Abstraction/BuiltInTweenFunctionContainer.cpp
@@ -2,6 +2,7 @@
 
 #include "Math/Abstraction/BuiltInTweenFunctionContainer.h"
 
+#include "Containers/UnrealString.h"
 #include "Framework/Commands/UICommandInfo.h"
 #include "Math/Abstraction/BuiltInTweenFunction.h"
 #include "Misc/Optional.h"
@@ -42,30 +43,36 @@ TSharedPtr<ITweenFunction> FBuiltInTweenFunctionContainer::GetFunction(int32 InI
 TSharedPtr<FUICommandInfo> FBuiltInTweenFunctionContainer::GetCommandForFunction(const TSharedRef<ITweenFunction>& InFunction) const
 {
 	const TOptional<EBlendFunction> Function = GetBlendFunctionType(InFunction);
-	return Function ? TweeningUtilsEditor::GetCommandForBlendFunction(*Function) : nullptr;
+	return Function ? GetCommandForBlendFunction(*Function) : nullptr;
 }
 
 const FSlateBrush* FBuiltInTweenFunctionContainer::GetIconForFunction(const TSharedRef<ITweenFunction>& InFunction) const
 {
 	const TOptional<EBlendFunction> Function = GetBlendFunctionType(InFunction);
-	return Function ? TweeningUtilsEditor::GetUntintedIconForTweenFunction(*Function) : nullptr;
+	return Function ? GetUntintedIconForTweenFunction(*Function) : nullptr;
 }
 
 FLinearColor FBuiltInTweenFunctionContainer::GetColorForFunction(const TSharedRef<ITweenFunction>& InFunction) const
 {
 	const TOptional<EBlendFunction> Function = GetBlendFunctionType(InFunction);
-	return Function ? TweeningUtilsEditor::GetTintColorForTweenFunction(*Function) : FLinearColor::White;
+	return Function ? GetTintColorForTweenFunction(*Function) : FLinearColor::White;
 }
 
 FText FBuiltInTweenFunctionContainer::GetLabelForFunction(const TSharedRef<ITweenFunction>& InFunction) const
 {
 	const TOptional<EBlendFunction> Function = GetBlendFunctionType(InFunction);
-	return Function ? TweeningUtilsEditor::GetLabelForBlendFunction(*Function) : FText::GetEmpty();
+	return Function ? GetLabelForBlendFunction(*Function) : FText::GetEmpty();
 }
 	
 FText FBuiltInTweenFunctionContainer::GetToolTipForFunction(const TSharedRef<ITweenFunction>& InFunction) const
 {
 	const TOptional<EBlendFunction> Function = GetBlendFunctionType(InFunction);
-	return Function ? TweeningUtilsEditor::GetDescriptionForBlendFunction(*Function) : FText::GetEmpty();
+	return Function ? GetDescriptionForBlendFunction(*Function) : FText::GetEmpty();
+}
+
+FString FBuiltInTweenFunctionContainer::GetFunctionIdentifier(const TSharedRef<ITweenFunction>& InFunction) const
+{
+	const TOptional<EBlendFunction> Function = GetBlendFunctionType(InFunction);
+	return Function ? BlendFunctionToString(*Function) : FString();
 }
 }
diff --git a/Engine/Plugins/Animation/TweeningUtils/Source/TweeningUtilsEditor/Private/Math/Abstraction/ITweenFunctionContainer.cpp b/Engine/Plugins/Animation/TweeningUtils/Source/TweeningUtilsEditor/Private/Math/Abstraction/ITweenFunctionContainer.cpp
--- a/Engine/Plugins/Animation/TweeningUtils/Source/TweeningUtilsEditor/Private/Math/Abstraction/ITweenFunctionContainer.cpp
+++ b/Engine/Plugins/Animation/TweeningUtils/Source/TweeningUtilsEditor/Private/Math/Abstraction/ITweenFunctionContainer.cpp
@@ -19,4 +19,17 @@ int32 ITweenFunctionContainer::IndexOf(const TSharedRef<ITweenFunction>& InFunct
 	}
 	return INDEX_NONE;
 }
+
+TSharedPtr<ITweenFunction> ITweenFunctionContainer::FindFunctionByIdentifier(const FString& InIdentifier) const
+{
+	for (int32 Index = 0; Index < NumFunctions(); ++Index)
+	{
+		const TSharedPtr<ITweenFunction> Function = GetFunction(Index);
+		if (GetFunctionIdentifier(Function.ToSharedRef()) == InIdentifier)
+		{
+			return Function;
+		}
+	}
+	return nullptr;
+}
 }
diff --git a/Engine/Plugins/Animation/TweeningUtils/Source/TweeningUtilsEditor/Private/Math/Abstraction/KeyBlendingAbstraction.cpp b/Engine/Plugins/Animation/TweeningUtils/Source/TweeningUtilsEditor/Private/Math/Abstraction/KeyBlendingAbstraction.cpp
--- a/Engine/Plugins/Animation/TweeningUtils/Source/TweeningUtilsEditor/Private/Math/Abstraction/KeyBlendingAbstraction.cpp
+++ b/Engine/Plugins/Animation/TweeningUtils/Source/TweeningUtilsEditor/Private/Math/Abstraction/KeyBlendingAbstraction.cpp
@@ -8,6 +8,7 @@
 #include "Math/Color.h"
 #include "Math/KeyBlendInfo.h"
 #include "Math/KeyBlendingFunctions.h"
+#include "Misc/Optional.h"
 #include "Styling/SlateBrush.h"
 #include "Templates/SharedPointer.h"
 #include "TweeningUtilsCommands.h"
@@ -27,6 +28,7 @@ struct FStaticBlendFunctionData
 	const FText Label;
 	const FText Abbreviation;
 	const FText Description;
+	const FString StringEncoding;
 	const FString StyleBaseName;
 	
 	FStaticBlendFunctionData(
@@ -35,7 +37,8 @@ struct FStaticBlendFunctionData
 		TSharedPtr<FUICommandInfo> InCommand,
 		FText InLabel,
 		FText InAbbreviation,
-		FText InDescription
+		FText InDescription,
+		FString InStringEncoding
 		)
 		: Type(InType)
 		, EvaluationFunctionPtr(InEvaluationFunctionPtr)
@@ -43,6 +46,7 @@ struct FStaticBlendFunctionData
 		, Label(MoveTemp(InLabel))
 		, Abbreviation(MoveTemp(InAbbreviation))
 		, Description(MoveTemp(InDescription))
+		, StringEncoding(MoveTemp(InStringEncoding))
 		, StyleBaseName(FString::Printf(TEXT("TweeningUtils.%s"), InCommand ? *InCommand->GetCommandName().ToString() : TEXT("")))
 	{
 		check(InCommand.IsValid());
@@ -71,43 +75,67 @@ static FStaticBlendFunctionData& GetFunctionData(EBlendFunction InFunction)
 			EBlendFunction::BlendNeighbor, &Blend_Neighbor, FTweeningUtilsCommands::Get().SetTweenBlendNeighbor,
 			LOCTEXT("BlendNeighbor.Label", "Blend Neighbor"),
 			LOCTEXT("BlendNeighbour.Abbreviation", "BN"),
-			LOCTEXT("BlendNeighbour.Description", "Blend to the next or previous values for selected keys.")
+			LOCTEXT("BlendNeighbour.Description", "Blend to the next or previous values for selected keys."),
+			TEXT("BlendNeighbour")
 		),
 		FStaticBlendFunctionData(
 			EBlendFunction::PushPull, &Blend_PushPull, FTweeningUtilsCommands::Get().SetTweenPushPull,
 			LOCTEXT("PushPull.Label", "Push / Pull"),
 			LOCTEXT("PushPull.Abbreviation", "PP"),
-			LOCTEXT("PushPull.Description", "Push or pull the values to the interpolation between the previous and next keys")
+			LOCTEXT("PushPull.Description", "Push or pull the values to the interpolation between the previous and next keys"),
+			TEXT("PushPull")
 		),
 		FStaticBlendFunctionData(
 			EBlendFunction::BlendEase, &Blend_Ease, FTweeningUtilsCommands::Get().SetTweenBlendEase,
 			LOCTEXT("BlendEase.Label", "Blend Ease"),
 			LOCTEXT("BlendEase.Abbreviation", "BE"),
-			LOCTEXT("BlendEase.Description", "Blend with an ease falloff to the next or previous value for selected keys")
+			LOCTEXT("BlendEase.Description", "Blend with an ease falloff to the next or previous value for selected keys"),
+			TEXT("BlendEase")
 		),
 		FStaticBlendFunctionData(
 			EBlendFunction::ControlsToTween, &Blend_ControlsToTween, FTweeningUtilsCommands::Get().SetControlsToTween,
 			LOCTEXT("ControlsToTween.Label", "Tween"),
 			LOCTEXT("ControlsToTween.Abbreviation", "TW"),
-			LOCTEXT("ControlsToTween.Description", "Interpolates between the previous and next keys")
+			LOCTEXT("ControlsToTween.Description", "Interpolates between the previous and next keys"),
+			TEXT("Tween")
 		),
 		FStaticBlendFunctionData(
 			EBlendFunction::BlendRelative, &Blend_Relative, FTweeningUtilsCommands::Get().SetTweenBlendRelative,
 			LOCTEXT("BlendRelative.Label", "Blend Relative"),
 			LOCTEXT("BlendRelative.Abbreviation", "BR"),
-			LOCTEXT("BlendRelative.Description", "Blend relative to the next or previous value for selected keys")
+			LOCTEXT("BlendRelative.Description", "Blend relative to the next or previous value for selected keys"),
+			TEXT("BlendRelative")
 		),
 		FStaticBlendFunctionData(
 			EBlendFunction::SmoothRough, &Blend_SmoothRough, FTweeningUtilsCommands::Get().SetTweenSmoothRough,
 			LOCTEXT("SmoothRough.Label", "Smooth Rough"),
 			LOCTEXT("SmoothRough.Abbreviation", "SR"),
-			LOCTEXT("SmoothRough.Description", "Make selected keys smooth or rough")
+			LOCTEXT("SmoothRough.Description", "Make selected keys smooth or rough"),
+			TEXT("SmoothRough")
 		)
 	};
 
 	const int32 Index = static_cast<int32>(InFunction);
 	return Data[Index];
 }
+
+FString BlendFunctionToString(EBlendFunction InFunction)
+{
+	return GetFunctionData(InFunction).StringEncoding;
+}
+
+TOptional<EBlendFunction> LexBlendFunction(const FString& InString)
+{
+	for (int32 FuncIndex = 0; FuncIndex < static_cast<int32>(EBlendFunction::Num); ++FuncIndex)
+	{
+		const EBlendFunction Function = static_cast<EBlendFunction>(FuncIndex);
+		if (InString == GetFunctionData(Function).StringEncoding)
+		{
+			return Function;
+		}
+	}
+	return {};
+}
 	
 double EvaluateBlend(EBlendFunction InFunction, const FKeyBlendInfo& InBlendInfo)
 {
diff --git a/Engine/Plugins/Animation/TweeningUtils/Source/TweeningUtilsEditor/Private/Widgets/MVC/TweenToolbarController.cpp b/Engine/Plugins/Animation/TweeningUtils/Source/TweeningUtilsEditor/Private/Widgets/MVC/TweenToolbarController.cpp
--- a/Engine/Plugins/Animation/TweeningUtils/Source/TweeningUtilsEditor/Private/Widgets/MVC/TweenToolbarController.cpp
+++ b/Engine/Plugins/Animation/TweeningUtils/Source/TweeningUtilsEditor/Private/Widgets/MVC/TweenToolbarController.cpp
@@ -104,7 +104,9 @@ FSlateIcon FTweenToolbarController::GetIconForComboBox() const
 
 void FTweenToolbarController::SetBlendFunction(TSharedPtr<ITweenFunction> InFunction) const
 {
-	TweenModel->SetTweenFunction(InFunction.ToSharedRef());
+	const TSharedRef<ITweenFunction> FunctionRef = InFunction.ToSharedRef();
+	TweenModel->SetTweenFunction(FunctionRef);
+	OnTweenFunctionChangedDelegate.Broadcast(FunctionRef);
 }
 
 bool FTweenToolbarController::IsBlendFunctionSelected(TSharedPtr<ITweenFunction> InFunction) const
diff --git a/Engine/Plugins/Animation/TweeningUtils/Source/TweeningUtilsEditor/Public/Math/Abstraction/BuiltInTweenFunctionContainer.h b/Engine/Plugins/Animation/TweeningUtils/Source/TweeningUtilsEditor/Public/Math/Abstraction/BuiltInTweenFunctionContainer.h
--- a/Engine/Plugins/Animation/TweeningUtils/Source/TweeningUtilsEditor/Public/Math/Abstraction/BuiltInTweenFunctionContainer.h
+++ b/Engine/Plugins/Animation/TweeningUtils/Source/TweeningUtilsEditor/Public/Math/Abstraction/BuiltInTweenFunctionContainer.h
@@ -30,6 +30,7 @@ public:
 	virtual FLinearColor GetColorForFunction(const TSharedRef<ITweenFunction>& InFunction) const override;
 	virtual FText GetLabelForFunction(const TSharedRef<ITweenFunction>& InFunction) const override;
 	virtual FText GetToolTipForFunction(const TSharedRef<ITweenFunction>& InFunction) const override;
+	virtual FString GetFunctionIdentifier(const TSharedRef<ITweenFunction>& InFunction) const override;
 	//~ Begin ITweenFunctionContainer Interface
 
 private:
diff --git a/Engine/Plugins/Animation/TweeningUtils/Source/TweeningUtilsEditor/Public/Math/Abstraction/ITweenFunctionContainer.h b/Engine/Plugins/Animation/TweeningUtils/Source/TweeningUtilsEditor/Public/Math/Abstraction/ITweenFunctionContainer.h
--- a/Engine/Plugins/Animation/TweeningUtils/Source/TweeningUtilsEditor/Public/Math/Abstraction/ITweenFunctionContainer.h
+++ b/Engine/Plugins/Animation/TweeningUtils/Source/TweeningUtilsEditor/Public/Math/Abstraction/ITweenFunctionContainer.h
@@ -5,6 +5,7 @@
 #include "Templates/FunctionFwd.h"
 #include "Templates/SharedPointerFwd.h"
 
+class FString;
 class FText;
 class FUICommandInfo;
 struct FLinearColor;
@@ -44,6 +45,11 @@ public:
 	/** @return The description to displayed, e.g. in tool tips. */
 	virtual FText GetToolTipForFunction(const TSharedRef<ITweenFunction>& InFunction) const = 0;
 
+	/** @return Identifier that uniquely identifies this function. Used e.g. to encode the function type in a config file. */
+	virtual FString GetFunctionIdentifier(const TSharedRef<ITweenFunction>& InFunction) const = 0;
+	/** @return The function identified by the InIdentifier (as returned by a call to GetFunctionIdentifier). */
+	TWEENINGUTILSEDITOR_API TSharedPtr<ITweenFunction> FindFunctionByIdentifier(const FString& InIdentifier) const;
+
 	virtual ~ITweenFunctionContainer() = default;
 };
 }
diff --git a/Engine/Plugins/Animation/TweeningUtils/Source/TweeningUtilsEditor/Public/Math/Abstraction/KeyBlendingAbstraction.h b/Engine/Plugins/Animation/TweeningUtils/Source/TweeningUtilsEditor/Public/Math/Abstraction/KeyBlendingAbstraction.h
--- a/Engine/Plugins/Animation/TweeningUtils/Source/TweeningUtilsEditor/Public/Math/Abstraction/KeyBlendingAbstraction.h
+++ b/Engine/Plugins/Animation/TweeningUtils/Source/TweeningUtilsEditor/Public/Math/Abstraction/KeyBlendingAbstraction.h
@@ -7,10 +7,12 @@
 
 #include <type_traits>
 
+class FString;
 class FText;
 class FUICommandInfo;
 struct FLinearColor;
 struct FSlateBrush;
+template<typename T> struct TOptional;
 
 /** The following just contains some utils that is common in UI code. Tools can opt for their own description but using this helps with consistency. */
 namespace UE::TweeningUtilsEditor
@@ -42,6 +44,11 @@ enum class EBlendFunction : uint8
 	Num
 };
 
+/** @return Converts InFunction to a string */
+TWEENINGUTILSEDITOR_API FString BlendFunctionToString(EBlendFunction InFunction);
+/** @return The EBlendFunction encoded by InString. */
+TWEENINGUTILSEDITOR_API TOptional<EBlendFunction> LexBlendFunction(const FString& InString);
+	
 /** Evaluates the function */
 TWEENINGUTILSEDITOR_API double EvaluateBlend(EBlendFunction InFunction, const FKeyBlendInfo& InBlendInfo);
 
diff --git a/Engine/Plugins/Animation/TweeningUtils/Source/TweeningUtilsEditor/Public/Widgets/MVC/TweenToolbarController.h b/Engine/Plugins/Animation/TweeningUtils/Source/TweeningUtilsEditor/Public/Widgets/MVC/TweenToolbarController.h
--- a/Engine/Plugins/Animation/TweeningUtils/Source/TweeningUtilsEditor/Public/Widgets/MVC/TweenToolbarController.h
+++ b/Engine/Plugins/Animation/TweeningUtils/Source/TweeningUtilsEditor/Public/Widgets/MVC/TweenToolbarController.h
@@ -55,6 +55,10 @@ public:
 	 */
 	void AddToToolbar(FToolBarBuilder& ToolbarBuilder, TAttribute<TOptional<float>> InOverrideSliderPosition = {}) const;
 
+	DECLARE_MULTICAST_DELEGATE_OneParam(FOnTweenFunctionChanged, const TSharedRef<ITweenFunction>&);
+	/** Invoked when the tween function is changed by this controller. */
+	FOnTweenFunctionChanged& OnTweenFunctionChanged() { return OnTweenFunctionChangedDelegate; }
+
 private:
 
 	/** Receives the blend values */
@@ -64,6 +68,9 @@ private:
 	/** Holds the functions that can be selected and gives us information about how to display them. */
 	const TSharedRef<ITweenFunctionContainer> TweenFunctions;
 
+	/** Invoked when the tween function is changed by this controller. */
+	FOnTweenFunctionChanged OnTweenFunctionChangedDelegate;
+
 	/** Creates the drop-down list for selecting the tween mode. */
 	TSharedRef<SWidget> MakeTweenModeMenu() const;
 	FText GetLabelForComboBox() const;
diff --git a/Engine/Plugins/Editor/CurveEditorTools/Source/CurveEditorTools/Private/CurveEditorToolsUserSettings.cpp b/Engine/Plugins/Editor/CurveEditorTools/Source/CurveEditorTools/Private/CurveEditorToolsUserSettings.cpp
new file mode 100644
--- /dev/null
+++ b/Engine/Plugins/Editor/CurveEditorTools/Source/CurveEditorTools/Private/CurveEditorToolsUserSettings.cpp
@@ -0,0 +1,10 @@
+﻿// Copyright Epic Games, Inc. All Rights Reserved.
+
+#include "CurveEditorToolsUserSettings.h"
+
+#include "UObject/UObjectGlobals.h"
+
+UCurveEditorToolsUserSettings* UCurveEditorToolsUserSettings::Get()
+{
+	return GetMutableDefault<UCurveEditorToolsUserSettings>();
+}
diff --git a/Engine/Plugins/Editor/CurveEditorTools/Source/CurveEditorTools/Private/CurveEditorToolsUserSettings.h b/Engine/Plugins/Editor/CurveEditorTools/Source/CurveEditorTools/Private/CurveEditorToolsUserSettings.h
new file mode 100644
--- /dev/null
+++ b/Engine/Plugins/Editor/CurveEditorTools/Source/CurveEditorTools/Private/CurveEditorToolsUserSettings.h
@@ -0,0 +1,27 @@
+﻿// Copyright Epic Games, Inc. All Rights Reserved.
+
+#pragma once
+
+#include "UObject/Object.h"
+#include "CurveEditorToolsUserSettings.generated.h"
+
+/**
+ * 
+ */
+UCLASS(Config=EditorPerProjectUserSettings)
+class CURVEEDITORTOOLS_API UCurveEditorToolsUserSettings : public UObject
+{
+	GENERATED_BODY()
+public:
+
+	static UCurveEditorToolsUserSettings* Get();
+
+	void SetPreferredTweenFunction(FString InNewPreferredFunction) { PreferredTweenFunction = MoveTemp(InNewPreferredFunction); SaveConfig(); }
+	const FString& GetPreferredTweenFunction() const { return PreferredTweenFunction; }
+	
+private:
+	
+	/** This tween function should be selected by default when a curve editor is created. It is the function that was used last time in the editor. */
+	UPROPERTY(Config)
+	FString PreferredTweenFunction;
+};
diff --git a/Engine/Plugins/Editor/CurveEditorTools/Source/CurveEditorTools/Private/Extensions/Tweening/TweenEditorExtension.cpp b/Engine/Plugins/Editor/CurveEditorTools/Source/CurveEditorTools/Private/Extensions/Tweening/TweenEditorExtension.cpp
--- a/Engine/Plugins/Editor/CurveEditorTools/Source/CurveEditorTools/Private/Extensions/Tweening/TweenEditorExtension.cpp
+++ b/Engine/Plugins/Editor/CurveEditorTools/Source/CurveEditorTools/Private/Extensions/Tweening/TweenEditorExtension.cpp
@@ -4,6 +4,7 @@
 
 #include "ContextAwareTweeningModel.h"
 #include "CurveEditor.h"
+#include "CurveEditorToolsUserSettings.h"
 #include "Framework/MultiBox/MultiBoxBuilder.h"
 #include "Framework/MultiBox/MultiBoxExtender.h"
 #include "Math/Abstraction/BuiltInTweenFunctionContainer.h"
@@ -74,13 +75,26 @@ bool FTweenEditorExtension::InitControllers(TSharedRef<FUICommandList> InCommand
 		using namespace TweeningUtilsEditor;
 		using FModel = FContextAwareTweeningModel;
 
-		const TSharedRef<ITweenFunction> InitialFunction = BlendFunctionContainer->GetFunction(0).ToSharedRef();
+		const TSharedPtr<ITweenFunction> PreferredFunction = BlendFunctionContainer->FindFunctionByIdentifier(
+			UCurveEditorToolsUserSettings::Get()->GetPreferredTweenFunction()
+			);
+		const TSharedRef<ITweenFunction> InitialFunction = PreferredFunction
+			? PreferredFunction.ToSharedRef() : BlendFunctionContainer->GetFunction(0).ToSharedRef();
+		
 		const TSharedRef<TTransactionalTweenModelProxy<FModel>> TweenModel = MakeShared<TTransactionalTweenModelProxy<FModel>>(
 			WeakCurveEditor, InitialFunction
 			);
 		TweenControllers.Emplace(TweenModel, InCommandList, BlendFunctionContainer);
+		TweenControllers->ToolbarController.OnTweenFunctionChanged().AddRaw(this, &FTweenEditorExtension::OnTweenFunctionSelected);
+		
 		return true;
 	}
 	return false;
 }
+
+void FTweenEditorExtension::OnTweenFunctionSelected(const TSharedRef<TweeningUtilsEditor::ITweenFunction>& InFunction) const
+{
+	const FString& FunctionConfigId = BlendFunctionContainer->GetFunctionIdentifier(InFunction);
+	UCurveEditorToolsUserSettings::Get()->SetPreferredTweenFunction(FunctionConfigId);
+}
 }
diff --git a/Engine/Plugins/Editor/CurveEditorTools/Source/CurveEditorTools/Private/Extensions/Tweening/TweenEditorExtension.h b/Engine/Plugins/Editor/CurveEditorTools/Source/CurveEditorTools/Private/Extensions/Tweening/TweenEditorExtension.h
--- a/Engine/Plugins/Editor/CurveEditorTools/Source/CurveEditorTools/Private/Extensions/Tweening/TweenEditorExtension.h
+++ b/Engine/Plugins/Editor/CurveEditorTools/Source/CurveEditorTools/Private/Extensions/Tweening/TweenEditorExtension.h
@@ -13,6 +13,7 @@ class FCurveEditor;
 class FExtender;
 namespace UE::TweeningUtilsEditor { class FBuiltInTweenFunctionContainer; }
 namespace UE::TweeningUtilsEditor { class FFunctionBasedTweenModel; }
+namespace UE::TweeningUtilsEditor { class ITweenFunction; }
 namespace UE::TweeningUtilsEditor { class ITweenFunctionContainer; }
 
 namespace UE::CurveEditorTools
@@ -40,13 +41,13 @@ private:
 	struct FTweenControllers
 	{
 		/** Manages the toolbar widget */
-		const TweeningUtilsEditor::FTweenToolbarController ToolbarController;
+		TweeningUtilsEditor::FTweenToolbarController ToolbarController;
 
 		/** Cycles functions (Shift + U). */
-		const TweeningUtilsEditor::FCycleFunctionController CycleFunctionController;
+		TweeningUtilsEditor::FCycleFunctionController CycleFunctionController;
 
 		/** Allows indirect movement of the slider by using u + Move Mouse */
-		const TweeningUtilsEditor::FTweenMouseSlidingController MouseSlidingController;
+		TweeningUtilsEditor::FTweenMouseSlidingController MouseSlidingController;
 		
 		explicit FTweenControllers(
 			const TSharedRef<TweeningUtilsEditor::FFunctionBasedTweenModel>& InTweenModel,
@@ -59,6 +60,9 @@ private:
 
 	/** Makes sure that ToolbarController is created. */
 	bool InitControllers(TSharedRef<FUICommandList> InCommandList);
+	
+	/** Invoked when the combo button selects a tween function. */
+	void OnTweenFunctionSelected(const TSharedRef<TweeningUtilsEditor::ITweenFunction>& Shared) const;
 };
 }
 
