Curve Editor: Make tween button black and make overshoot button have same color as the tween function's color.

diff --git a/Engine/Plugins/Animation/TweeningUtils/Source/TweeningUtilsEditor/Private/Widgets/MVC/TweenToolbarController.cpp b/Engine/Plugins/Animation/TweeningUtils/Source/TweeningUtilsEditor/Private/Widgets/MVC/TweenToolbarController.cpp
--- a/Engine/Plugins/Animation/TweeningUtils/Source/TweeningUtilsEditor/Private/Widgets/MVC/TweenToolbarController.cpp
+++ b/Engine/Plugins/Animation/TweeningUtils/Source/TweeningUtilsEditor/Private/Widgets/MVC/TweenToolbarController.cpp
@@ -6,12 +6,53 @@
 #include "Framework/MultiBox/MultiBoxBuilder.h"
 #include "Math/Abstraction/ITweenFunctionContainer.h"
 #include "Math/Color.h"
+#include "Styling/SlateStyle.h"
+#include "Styling/SlateStyleMacros.h"
+#include "Styling/SlateStyleRegistry.h"
 #include "TweeningUtilsCommands.h"
 #include "Widgets/MVC/Models/FunctionBasedTweenModel.h"
 #include "Widgets/MVC/STweenView.h"
 
 namespace UE::TweeningUtilsEditor
 {
+class FOvershootButtonStyleHack : public FSlateStyleSet
+{
+public:
+
+	explicit FOvershootButtonStyleHack(const ITweenFunctionContainer& InTweenFunctions)
+		: FSlateStyleSet(*FString::Printf(TEXT("TweeningUtils_OvershootOverride_%s"), *FGuid::NewGuid().ToString()))
+	{
+		const FString PluginContentDir = FPaths::EnginePluginsDir() / TEXT("Animation") / TEXT("TweeningUtils") / TEXT("Resources");
+		const FString EngineEditorSlateDir = FPaths::EngineContentDir() / TEXT("Editor") / TEXT("Slate");
+		FSlateStyleSet::SetContentRoot(PluginContentDir);
+		FSlateStyleSet::SetCoreContentRoot(EngineEditorSlateDir);
+		
+		const FVector2D Icon20x20(20.0f, 20.0f);
+		FSlateBrush* BaseBrush = new IMAGE_BRUSH_SVG("Icons/OvershootMode", Icon20x20);
+		Set(TEXT("OvershootMode"), BaseBrush);
+		InTweenFunctions.ForEachFunction([this, &InTweenFunctions, &BaseBrush](const TSharedRef<ITweenFunction>& Function) 
+		{
+			FSlateBrush* Brush = new FSlateBrush(*BaseBrush);
+			Brush->TintColor = InTweenFunctions.GetColorForFunction(Function);
+			Set(*GetOvershootButtonStyleName(InTweenFunctions, Function), Brush);
+		});
+		
+		FSlateStyleRegistry::RegisterSlateStyle(*this);
+	}
+
+	virtual ~FOvershootButtonStyleHack() override
+	{
+		FSlateStyleRegistry::UnRegisterSlateStyle(*this);
+	}
+
+	static FString GetOvershootButtonStyleName(
+		const ITweenFunctionContainer& InTweenFunctions, const TSharedRef<ITweenFunction>& InFunction
+		)
+	{
+		return InTweenFunctions.GetFunctionIdentifier(InFunction);
+	}
+};
+	
 FTweenToolbarController::FTweenToolbarController(
 	const TSharedRef<FFunctionBasedTweenModel>& InModel,
 	const TSharedRef<FUICommandList>& InCommandList,
@@ -20,6 +61,7 @@ FTweenToolbarController::FTweenToolbarController(
 	: TweenModel(InModel)
 	, CommandList(InCommandList)
 	, TweenFunctions(InTweenFunctions)
+	, OverrideStyle(MakePimpl<FOvershootButtonStyleHack>(*InTweenFunctions))
 {
 	// Function changing commands
 	TweenFunctions->ForEachFunction([this](const TSharedRef<ITweenFunction>& InFunction)
@@ -72,7 +114,11 @@ void FTweenToolbarController::AddToToolbar(FToolBarBuilder& ToolBarBuilder, TAtt
 		.SliderColor(TAttribute<FLinearColor>::CreateRaw(this, &FTweenToolbarController::GetColorForSlider))
 		.OverrideSliderPosition(MoveTemp(InOverrideSliderPosition))
 		);
-	ToolBarBuilder.AddToolBarButton(FTweeningUtilsCommands::Get().ToggleOvershootMode);
+	
+	ToolBarBuilder.AddToolBarButton(FTweeningUtilsCommands::Get().ToggleOvershootMode,
+		NAME_None, TAttribute<FText>(), TAttribute<FText>(),
+		TAttribute<FSlateIcon>::CreateRaw(this, &FTweenToolbarController::GetOvershootModeIcon)
+		);
 }
 
 TSharedRef<SWidget> FTweenToolbarController::MakeTweenModeMenu() const
@@ -126,6 +172,14 @@ bool FTweenToolbarController::IsOvershootModeEnabled() const
 	return TweenModel->GetScaleMode() == ETweenScaleMode::Overshoot;
 }
 
+FSlateIcon FTweenToolbarController::GetOvershootModeIcon() const
+{
+	return FSlateIcon(
+		OverrideStyle->GetStyleSetName(),
+		IsOvershootModeEnabled() ? TEXT("OvershootMode") : *OverrideStyle->GetOvershootButtonStyleName(*TweenFunctions, TweenModel->GetTweenFunction())
+		);
+}
+
 const FSlateBrush* FTweenToolbarController::GetIconForSlider() const
 {
 	return TweenFunctions->GetIconForFunction(TweenModel->GetTweenFunction());
diff --git a/Engine/Plugins/Animation/TweeningUtils/Source/TweeningUtilsEditor/Private/Widgets/TweenSliderDrawUtils.cpp b/Engine/Plugins/Animation/TweeningUtils/Source/TweeningUtilsEditor/Private/Widgets/TweenSliderDrawUtils.cpp
--- a/Engine/Plugins/Animation/TweeningUtils/Source/TweeningUtilsEditor/Private/Widgets/TweenSliderDrawUtils.cpp
+++ b/Engine/Plugins/Animation/TweeningUtils/Source/TweeningUtilsEditor/Private/Widgets/TweenSliderDrawUtils.cpp
@@ -360,7 +360,7 @@ static void DrawButton(
 
 	const FSlateColor& BaseIconTint = InDrawArgs.bDrawButtonPressed
 		? Style->PressedIconTint
-		: InDrawArgs.HoverState.bIsSliderHovered ? Style->HoveredIconTint : Style->PressedIconTint;
+		: InDrawArgs.HoverState.bIsSliderHovered ? Style->HoveredIconTint : Style->NormalIconTint;
 	
 	const FLinearColor ColorAndOpacity = InWidgetArgs.ColorAndOpacity.Get().GetColor(InWidgetStyle);
 	const FLinearColor ButtonTint = ColorAndOpacity * InWidgetArgs.SliderColor.Get() * SliderButtonBrush->TintColor.GetColor(InWidgetStyle);
diff --git a/Engine/Plugins/Animation/TweeningUtils/Source/TweeningUtilsEditor/Private/Widgets/TweenSliderStyle.cpp b/Engine/Plugins/Animation/TweeningUtils/Source/TweeningUtilsEditor/Private/Widgets/TweenSliderStyle.cpp
--- a/Engine/Plugins/Animation/TweeningUtils/Source/TweeningUtilsEditor/Private/Widgets/TweenSliderStyle.cpp
+++ b/Engine/Plugins/Animation/TweeningUtils/Source/TweeningUtilsEditor/Private/Widgets/TweenSliderStyle.cpp
@@ -36,7 +36,7 @@ FTweenSliderStyle::FTweenSliderStyle()
 	, NormalSliderButton(static_cast<FSlateBrush>(FSlateRoundedBoxBrush(FStyleColors::Foreground, 5.0f, FStyleColors::Transparent, 0.f)))
 	, HoveredSliderButton(static_cast<FSlateBrush>(FSlateRoundedBoxBrush(FStyleColors::ForegroundHover, 5.0f, FStyleColors::Transparent, 0.f)))
 	, PressedSliderButton(static_cast<FSlateBrush>(FSlateRoundedBoxBrush(FStyleColors::Foreground, 5.0f, FStyleColors::Transparent, 0.f)))
-	, NormalIconTint(FStyleColors::Foreground)
+	, NormalIconTint(FStyleColors::Black)
 	, HoveredIconTint(FStyleColors::White)
 	, PressedIconTint(FStyleColors::White)
 	, SmallPoint({ 4.0, 4.0 }, { 6.0, 6.0 }, { 6.0, 6.0 }, { 10.0, 12.0 })
diff --git a/Engine/Plugins/Animation/TweeningUtils/Source/TweeningUtilsEditor/Public/Widgets/MVC/TweenToolbarController.h b/Engine/Plugins/Animation/TweeningUtils/Source/TweeningUtilsEditor/Public/Widgets/MVC/TweenToolbarController.h
--- a/Engine/Plugins/Animation/TweeningUtils/Source/TweeningUtilsEditor/Public/Widgets/MVC/TweenToolbarController.h
+++ b/Engine/Plugins/Animation/TweeningUtils/Source/TweeningUtilsEditor/Public/Widgets/MVC/TweenToolbarController.h
@@ -3,6 +3,7 @@
 #pragma once
 
 #include "Misc/Attribute.h"
+#include "Templates/PimplPtr.h"
 #include "Templates/SharedPointer.h"
 #include "Templates/UnrealTemplate.h"
 
@@ -45,7 +46,7 @@ public:
 		const TSharedRef<ITweenFunctionContainer>& InTweenFunctions
 		);
 	~FTweenToolbarController();
-	
+
 	/**
 	 * Extends the toolbar. Safe to call multiple times (although typically you'll call this exactly once).
 	 * 
@@ -68,6 +69,9 @@ private:
 	/** Holds the functions that can be selected and gives us information about how to display them. */
 	const TSharedRef<ITweenFunctionContainer> TweenFunctions;
 
+	/** Used to override the color of the overshoot button to match the color of the function. */
+	const TPimplPtr<class FOvershootButtonStyleHack> OverrideStyle;
+
 	/** Invoked when the tween function is changed by this controller. */
 	FOnTweenFunctionChanged OnTweenFunctionChangedDelegate;
 
@@ -82,6 +86,7 @@ private:
 	
 	void ToggleOvershootMode() const;
 	bool IsOvershootModeEnabled() const;
+	FSlateIcon GetOvershootModeIcon() const;
 	
 	const FSlateBrush* GetIconForSlider() const;
 	FLinearColor GetColorForSlider() const;
