Tweening: Refactor tween functions to be referenced by const ITweenFunction instead of TSharedRef<ITweenFunction>. This change is more memory friendly so code, such as ControlRigSequencerEditorLibrary.cpp can use the tweening framework in one-off scenarious without needlessly forcing memory allocations.
- FFunctionBasedTweenModel: references functions from the function container
- FBuiltInTweenFunctionContainer: contains the function

diff --git a/Engine/Plugins/Animation/ControlRig/Source/ControlRigEditor/Private/ControlRigSequencerEditorLibrary.cpp b/Engine/Plugins/Animation/ControlRig/Source/ControlRigEditor/Private/ControlRigSequencerEditorLibrary.cpp
--- a/Engine/Plugins/Animation/ControlRig/Source/ControlRigEditor/Private/ControlRigSequencerEditorLibrary.cpp
+++ b/Engine/Plugins/Animation/ControlRig/Source/ControlRigEditor/Private/ControlRigSequencerEditorLibrary.cpp
@@ -105,7 +105,6 @@ TArray<FControlRigSequencerBindingProxy> UControlRigSequencerEditorLibrary::GetC
 static TArray<UObject*> GetBoundObjects(UWorld* World, ULevelSequence* LevelSequence, const FMovieSceneBindingProxy& Binding, ULevelSequencePlayer** OutPlayer, ALevelSequenceActor** OutActor)
 {
 	FMovieSceneSequencePlaybackSettings Settings;
-	FLevelSequenceCameraSettings CameraSettings;
 
 	ULevelSequencePlayer* Player = ULevelSequencePlayer::CreateLevelSequencePlayer(World, LevelSequence, Settings, *OutActor);
 	*OutPlayer = Player;
@@ -413,7 +412,7 @@ bool UControlRigSequencerEditorLibrary::TweenControlRig(ULevelSequence* LevelSeq
 		
 		// TODO UE-240486: This has unnecessary dynamic memory allocation overhead until UE-240486 is addressed.
 		const FBuiltInTweenFunctionContainer TweenFunctionContainer; 
-		FControlRigTweenModel BlendTool(WeakSequencer, TweenFunctionContainer.GetBlendFunction(EBlendFunction::ControlsToTween).ToSharedRef());
+		FControlRigTweenModel BlendTool(WeakSequencer, *TweenFunctionContainer.GetBlendFunction(EBlendFunction::ControlsToTween));
 		BlendTool.StartBlendOperationWithControlRigs({ ControlRig });
 		BlendTool.BlendValues(TweenValue);
 		BlendTool.StopBlendOperation();
@@ -457,7 +456,7 @@ bool UControlRigSequencerEditorLibrary::BlendValuesOnSelected(ULevelSequence* Le
 			const FBuiltInTweenFunctionContainer TweenFunctionContainer;
 			const EBlendFunction BlendFunction = SequencerLibraryDetail::ConvertBlendFunction(BlendOperation);
 			FControlRigTweenModel BlendTool(WeakSequencer, WeakMode,
-				TweenFunctionContainer.GetBlendFunction(BlendFunction).ToSharedRef()
+				*TweenFunctionContainer.GetBlendFunction(BlendFunction)
 				);
 			return BlendTool.BlendSingle(BlendValue);
 		}
@@ -1213,7 +1212,6 @@ bool UControlRigSequencerEditorLibrary::BakeToControlRig(UWorld* World, ULevelSe
 
 	ALevelSequenceActor* OutActor = nullptr;
 	FMovieSceneSequencePlaybackSettings Settings;
-	FLevelSequenceCameraSettings CameraSettings;
 	FMovieSceneSequenceIDRef Template = MovieSceneSequenceID::Root;
 	FMovieSceneSequenceTransform RootToLocalTransform;
 	IMovieScenePlayer* Player = nullptr;
@@ -1513,7 +1511,6 @@ static bool GetControlRigValues(UWorld* World, ULevelSequence* LevelSequence, UC
 	{
 		ALevelSequenceActor* OutActor = nullptr;
 		FMovieSceneSequencePlaybackSettings Settings;
-		FLevelSequenceCameraSettings CameraSettings;
 		FMovieSceneSequenceIDRef Template = MovieSceneSequenceID::Root;
 		FMovieSceneSequenceTransform RootToLocalTransform;
 		ULevelSequencePlayer* Player = ULevelSequencePlayer::CreateLevelSequencePlayer(World, LevelSequence, Settings, OutActor);
@@ -2726,7 +2723,6 @@ bool UControlRigSequencerEditorLibrary::ImportFBXToControlRigTrack(UWorld* World
 	bool bValid = false;
 	ALevelSequenceActor* OutActor;
 	FMovieSceneSequencePlaybackSettings Settings;
-	FLevelSequenceCameraSettings CameraSettings;
 	ULevelSequencePlayer* Player = ULevelSequencePlayer::CreateLevelSequencePlayer(World, Sequence, Settings, OutActor);
 
 	INodeAndChannelMappings* ChannelMapping = Cast<INodeAndChannelMappings>(InTrack);
diff --git a/Engine/Plugins/Animation/ControlRig/Source/ControlRigEditor/Private/EditMode/Tween/ControlRigTweenModel.cpp b/Engine/Plugins/Animation/ControlRig/Source/ControlRigEditor/Private/EditMode/Tween/ControlRigTweenModel.cpp
--- a/Engine/Plugins/Animation/ControlRig/Source/ControlRigEditor/Private/EditMode/Tween/ControlRigTweenModel.cpp
+++ b/Engine/Plugins/Animation/ControlRig/Source/ControlRigEditor/Private/EditMode/Tween/ControlRigTweenModel.cpp
@@ -71,7 +71,7 @@ bool FControlRigTweenModel::DoBlend(float InNormalizedValue)
 	{
 		return TweeningUtilsEditor::BlendCurves(ScaledBlendValue, *CurveEditor, KeySelection, [this](const TweeningUtilsEditor::FKeyBlendInfo& InBlendInfo)
 		{
-			return GetTweenFunction()->Blend(InBlendInfo);
+			return GetTweenFunction().Blend(InBlendInfo);
 		});
 	}
 
@@ -104,7 +104,7 @@ bool FControlRigTweenModel::DoBlend(float InNormalizedValue)
 					const double CurrentTime = TickResolution.AsSeconds(FFrameTime(ObjectChannels.KeyBounds[Index].CurrentFrame));
 					BlendStruct.SetValues(PreviousTime, PreviousValue, CurrentTime, CurrentValue,
 						NextTime, NextValue, ScaledBlendValue, CurrentValue, CurrentValue, CurrentIndex);
-					const double NewValue = GetTweenFunction()->Blend(BlendStruct);
+					const double NewValue = GetTweenFunction().Blend(BlendStruct);
 					using namespace UE::MovieScene;
 					if (ObjectChannels.KeyBounds[Index].FloatChannel)
 					{
diff --git a/Engine/Plugins/Animation/ControlRig/Source/ControlRigEditor/Private/EditMode/Tween/ControlRigTweenModel.h b/Engine/Plugins/Animation/ControlRig/Source/ControlRigEditor/Private/EditMode/Tween/ControlRigTweenModel.h
--- a/Engine/Plugins/Animation/ControlRig/Source/ControlRigEditor/Private/EditMode/Tween/ControlRigTweenModel.h
+++ b/Engine/Plugins/Animation/ControlRig/Source/ControlRigEditor/Private/EditMode/Tween/ControlRigTweenModel.h
@@ -3,6 +3,7 @@
 #pragma once
 
 #include "ControlRigObjectSelection.h"
+#include "Math/ContiguousKeyMapping.h"
 #include "Misc/Attribute.h"
 #include "Templates/SharedPointer.h"
 #include "Widgets/MVC/Models/FunctionBasedTweenModel.h"
@@ -25,14 +26,14 @@ public:
 	
 	explicit FControlRigTweenModel(
 		TAttribute<TWeakPtr<ISequencer>> InSequencerAttr, TWeakPtr<FControlRigEditMode> InControlRigMode,
-		const TSharedRef<TweeningUtilsEditor::ITweenFunction>& InInitialFunction
+		const TweeningUtilsEditor::ITweenFunction& InInitialFunction
 		)
 		: FFunctionBasedTweenModel(InInitialFunction)
 		, SequencerAttr(MoveTemp(InSequencerAttr))
 		, ControlRigMode(MoveTemp(InControlRigMode))
 	{}
 	explicit FControlRigTweenModel(
-		TAttribute<TWeakPtr<ISequencer>> InSequencerAttr, const TSharedRef<TweeningUtilsEditor::ITweenFunction>& InInitialFunction
+		TAttribute<TWeakPtr<ISequencer>> InSequencerAttr, const TweeningUtilsEditor::ITweenFunction& InInitialFunction
 		)
 		: FFunctionBasedTweenModel(InInitialFunction)
 		, SequencerAttr(MoveTemp(InSequencerAttr))
diff --git a/Engine/Plugins/Animation/ControlRig/Source/ControlRigEditor/Private/EditMode/Tween/TweenLogic.cpp b/Engine/Plugins/Animation/ControlRig/Source/ControlRigEditor/Private/EditMode/Tween/TweenLogic.cpp
--- a/Engine/Plugins/Animation/ControlRig/Source/ControlRigEditor/Private/EditMode/Tween/TweenLogic.cpp
+++ b/Engine/Plugins/Animation/ControlRig/Source/ControlRigEditor/Private/EditMode/Tween/TweenLogic.cpp
@@ -28,7 +28,7 @@ static TSharedRef<FControlRigTweenModel> MakeTweenModel(
 	// TODO UE-240484: Read preferred function from config
 	return MakeShared<FModel>(
 		MoveTemp(CurveEditorAttr),	// TTangentFlatteningTweenProxy arg
-		MoveTemp(InSequencerAttr), InOwningEditMode, InFunctions.GetFunction(0).ToSharedRef()
+		MoveTemp(InSequencerAttr), InOwningEditMode, *InFunctions.GetFunction(0)
 		);
 }
 }
diff --git a/Engine/Plugins/Animation/TweeningUtils/Source/TweeningUtilsEditor/Private/Math/Abstraction/BuiltInTweenFunctionContainer.cpp b/Engine/Plugins/Animation/TweeningUtils/Source/TweeningUtilsEditor/Private/Math/Abstraction/BuiltInTweenFunctionContainer.cpp
--- a/Engine/Plugins/Animation/TweeningUtils/Source/TweeningUtilsEditor/Private/Math/Abstraction/BuiltInTweenFunctionContainer.cpp
+++ b/Engine/Plugins/Animation/TweeningUtils/Source/TweeningUtilsEditor/Private/Math/Abstraction/BuiltInTweenFunctionContainer.cpp
@@ -11,78 +11,82 @@
 namespace UE::TweeningUtilsEditor
 {
 FBuiltInTweenFunctionContainer::FBuiltInTweenFunctionContainer()
-{
-	ForEachBlendFunction([this](EBlendFunction BlendFunction){ Functions.Add(MakeShared<FBuiltInTweenFunction>(BlendFunction)); });
-}
+	: Functions([]()
+	{
+		FTweenArray Array;
+		ForEachBlendFunction([&Array](EBlendFunction BlendFunction){ Array.Emplace(BlendFunction); });
+		return Array;
+	}())
+{}
 
-TOptional<EBlendFunction> FBuiltInTweenFunctionContainer::GetBlendFunctionType(const TSharedRef<ITweenFunction>& InFunctionToFind) const
+TOptional<EBlendFunction> FBuiltInTweenFunctionContainer::GetBlendFunctionType(const ITweenFunction& InFunctionToFind) const
 {
-	const int32 Index = Functions.IndexOfByPredicate([&InFunctionToFind](const TSharedRef<FBuiltInTweenFunction>& Function)
+	const int32 Index = Functions.IndexOfByPredicate([&InFunctionToFind](const FBuiltInTweenFunction& Function)
 	{
-		return &InFunctionToFind.Get() == &Function.Get();
+		return &InFunctionToFind == &Function;
 	});
 	// API user can theoretically pass in ANY function so handle that case
 	return ensureMsgf(Functions.IsValidIndex(Index), TEXT("You passed a function not managed by this container"))
-		? Functions[Index]->GetBlendFunction()
+		? Functions[Index].GetBlendFunction()
 		: TOptional<EBlendFunction>();
 }
 
-TSharedPtr<ITweenFunction> FBuiltInTweenFunctionContainer::GetBlendFunction(EBlendFunction InBlendFunction) const
+const ITweenFunction* FBuiltInTweenFunctionContainer::GetBlendFunction(EBlendFunction InBlendFunction) const
 {
-	for (const TSharedRef<FBuiltInTweenFunction>& Function : Functions)
+	for (const FBuiltInTweenFunction& Function : Functions)
 	{
-		if (Function->GetBlendFunction() == InBlendFunction)
+		if (Function.GetBlendFunction() == InBlendFunction)
 		{
-			return Function;
+			return &Function;
 		}
 	}
 	return nullptr;
 }
 
-void FBuiltInTweenFunctionContainer::ForEachFunction(const TFunctionRef<void(const TSharedRef<ITweenFunction>&)> InConsumer) const
+void FBuiltInTweenFunctionContainer::ForEachFunction(const TFunctionRef<void(const ITweenFunction&)> InConsumer) const
 {
-	for (const TSharedRef<FBuiltInTweenFunction>& Function : Functions)
+	for (const FBuiltInTweenFunction& Function : Functions)
 	{
 		InConsumer(Function);
 	}
 }
 
-TSharedPtr<ITweenFunction> FBuiltInTweenFunctionContainer::GetFunction(int32 InIndex) const
+const ITweenFunction* FBuiltInTweenFunctionContainer::GetFunction(int32 InIndex) const
 {
-	return Functions[InIndex];
+	return &Functions[InIndex];
 }
 
-TSharedPtr<FUICommandInfo> FBuiltInTweenFunctionContainer::GetCommandForFunction(const TSharedRef<ITweenFunction>& InFunction) const
+TSharedPtr<FUICommandInfo> FBuiltInTweenFunctionContainer::GetCommandForFunction(const ITweenFunction& InFunction) const
 {
 	const TOptional<EBlendFunction> Function = GetBlendFunctionType(InFunction);
 	return Function ? GetCommandForBlendFunction(*Function) : nullptr;
 }
 
-const FSlateBrush* FBuiltInTweenFunctionContainer::GetIconForFunction(const TSharedRef<ITweenFunction>& InFunction) const
+const FSlateBrush* FBuiltInTweenFunctionContainer::GetIconForFunction(const ITweenFunction& InFunction) const
 {
 	const TOptional<EBlendFunction> Function = GetBlendFunctionType(InFunction);
 	return Function ? GetUntintedIconForTweenFunction(*Function) : nullptr;
 }
 
-FLinearColor FBuiltInTweenFunctionContainer::GetColorForFunction(const TSharedRef<ITweenFunction>& InFunction) const
+FLinearColor FBuiltInTweenFunctionContainer::GetColorForFunction(const ITweenFunction& InFunction) const
 {
 	const TOptional<EBlendFunction> Function = GetBlendFunctionType(InFunction);
 	return Function ? GetTintColorForTweenFunction(*Function) : FLinearColor::White;
 }
 
-FText FBuiltInTweenFunctionContainer::GetLabelForFunction(const TSharedRef<ITweenFunction>& InFunction) const
+FText FBuiltInTweenFunctionContainer::GetLabelForFunction(const ITweenFunction& InFunction) const
 {
 	const TOptional<EBlendFunction> Function = GetBlendFunctionType(InFunction);
 	return Function ? GetLabelForBlendFunction(*Function) : FText::GetEmpty();
 }
 	
-FText FBuiltInTweenFunctionContainer::GetToolTipForFunction(const TSharedRef<ITweenFunction>& InFunction) const
+FText FBuiltInTweenFunctionContainer::GetToolTipForFunction(const ITweenFunction& InFunction) const
 {
 	const TOptional<EBlendFunction> Function = GetBlendFunctionType(InFunction);
 	return Function ? GetDescriptionForBlendFunction(*Function) : FText::GetEmpty();
 }
 
-FString FBuiltInTweenFunctionContainer::GetFunctionIdentifier(const TSharedRef<ITweenFunction>& InFunction) const
+FString FBuiltInTweenFunctionContainer::GetFunctionIdentifier(const ITweenFunction& InFunction) const
 {
 	const TOptional<EBlendFunction> Function = GetBlendFunctionType(InFunction);
 	return Function ? BlendFunctionToString(*Function) : FString();
diff --git a/Engine/Plugins/Animation/TweeningUtils/Source/TweeningUtilsEditor/Private/Math/Abstraction/ITweenFunctionContainer.cpp b/Engine/Plugins/Animation/TweeningUtils/Source/TweeningUtilsEditor/Private/Math/Abstraction/ITweenFunctionContainer.cpp
--- a/Engine/Plugins/Animation/TweeningUtils/Source/TweeningUtilsEditor/Private/Math/Abstraction/ITweenFunctionContainer.cpp
+++ b/Engine/Plugins/Animation/TweeningUtils/Source/TweeningUtilsEditor/Private/Math/Abstraction/ITweenFunctionContainer.cpp
@@ -8,11 +8,11 @@
 
 namespace UE::TweeningUtilsEditor
 {
-int32 ITweenFunctionContainer::IndexOf(const TSharedRef<ITweenFunction>& InFunction) const
+int32 ITweenFunctionContainer::IndexOf(const ITweenFunction& InFunction) const
 {
 	for (int32 Index = 0; Index < NumFunctions(); ++Index)
 	{
-		if (GetFunction(Index) == InFunction)
+		if (GetFunction(Index) == &InFunction)
 		{
 			return Index;
 		}
@@ -20,12 +20,12 @@ int32 ITweenFunctionContainer::IndexOf(const TSharedRef<ITweenFunction>& InFunct
 	return INDEX_NONE;
 }
 
-TSharedPtr<ITweenFunction> ITweenFunctionContainer::FindFunctionByIdentifier(const FString& InIdentifier) const
+const ITweenFunction* ITweenFunctionContainer::FindFunctionByIdentifier(const FString& InIdentifier) const
 {
 	for (int32 Index = 0; Index < NumFunctions(); ++Index)
 	{
-		const TSharedPtr<ITweenFunction> Function = GetFunction(Index);
-		if (GetFunctionIdentifier(Function.ToSharedRef()) == InIdentifier)
+		const ITweenFunction* Function = GetFunction(Index);
+		if (Function && GetFunctionIdentifier(*Function) == InIdentifier)
 		{
 			return Function;
 		}
diff --git a/Engine/Plugins/Animation/TweeningUtils/Source/TweeningUtilsEditor/Private/Widgets/MVC/CycleFunctionController.cpp b/Engine/Plugins/Animation/TweeningUtils/Source/TweeningUtilsEditor/Private/Widgets/MVC/CycleFunctionController.cpp
--- a/Engine/Plugins/Animation/TweeningUtils/Source/TweeningUtilsEditor/Private/Widgets/MVC/CycleFunctionController.cpp
+++ b/Engine/Plugins/Animation/TweeningUtils/Source/TweeningUtilsEditor/Private/Widgets/MVC/CycleFunctionController.cpp
@@ -44,7 +44,7 @@ void FCycleFunctionController::CycleToNextFunction()
 	if (ensure(TweenFunctions->IsValidIndex(Index)))
 	{
 		const int32 NextIndex = TweenFunctions->IsValidIndex(Index + 1) ? Index + 1 : 0;
-		TweenModel->SetTweenFunction(TweenFunctions->GetFunction(NextIndex).ToSharedRef());
+		TweenModel->SetTweenFunction(*TweenFunctions->GetFunction(NextIndex));
 	}
 }
 }
diff --git a/Engine/Plugins/Animation/TweeningUtils/Source/TweeningUtilsEditor/Private/Widgets/MVC/Models/CurveTweeningModel.cpp b/Engine/Plugins/Animation/TweeningUtils/Source/TweeningUtilsEditor/Private/Widgets/MVC/Models/CurveTweeningModel.cpp
--- a/Engine/Plugins/Animation/TweeningUtils/Source/TweeningUtilsEditor/Private/Widgets/MVC/Models/CurveTweeningModel.cpp
+++ b/Engine/Plugins/Animation/TweeningUtils/Source/TweeningUtilsEditor/Private/Widgets/MVC/Models/CurveTweeningModel.cpp
@@ -7,7 +7,7 @@
 
 namespace UE::TweeningUtilsEditor
 {
-FCurveTweenModel::FCurveTweenModel(TWeakPtr<FCurveEditor> InWeakCurveEditor, const TSharedRef<ITweenFunction>& InTweenFunction)
+FCurveTweenModel::FCurveTweenModel(TWeakPtr<FCurveEditor> InWeakCurveEditor, const ITweenFunction& InTweenFunction)
 	: FFunctionBasedTweenModel(InTweenFunction)
 	, WeakCurveEditor(InWeakCurveEditor)
 {}
@@ -42,7 +42,7 @@ void FCurveTweenModel::BlendValues(float InNormalizedValue)
 	const float ScaledBlendValue = ScaleBlendValue(InNormalizedValue); 
 	TweeningUtilsEditor::BlendCurves(ScaledBlendValue, *CurveEditor, ContiguousKeySelection, [this](const FKeyBlendInfo& InBlendInfo)
 	{
-		return GetTweenFunction()->Blend(InBlendInfo);
+		return GetTweenFunction().Blend(InBlendInfo);
 	});
 }
 }
diff --git a/Engine/Plugins/Animation/TweeningUtils/Source/TweeningUtilsEditor/Private/Widgets/MVC/TweenToolbarController.cpp b/Engine/Plugins/Animation/TweeningUtils/Source/TweeningUtilsEditor/Private/Widgets/MVC/TweenToolbarController.cpp
--- a/Engine/Plugins/Animation/TweeningUtils/Source/TweeningUtilsEditor/Private/Widgets/MVC/TweenToolbarController.cpp
+++ b/Engine/Plugins/Animation/TweeningUtils/Source/TweeningUtilsEditor/Private/Widgets/MVC/TweenToolbarController.cpp
@@ -30,7 +30,7 @@ public:
 		const FVector2D Icon20x20(20.0f, 20.0f);
 		FSlateBrush* BaseBrush = new IMAGE_BRUSH_SVG("Icons/OvershootMode", Icon20x20);
 		Set(TEXT("OvershootMode"), BaseBrush);
-		InTweenFunctions.ForEachFunction([this, &InTweenFunctions, &BaseBrush](const TSharedRef<ITweenFunction>& Function) 
+		InTweenFunctions.ForEachFunction([this, &InTweenFunctions, &BaseBrush](const ITweenFunction& Function) 
 		{
 			FSlateBrush* Brush = new FSlateBrush(*BaseBrush);
 			Brush->TintColor = InTweenFunctions.GetColorForFunction(Function);
@@ -46,7 +46,7 @@ public:
 	}
 
 	static FString GetOvershootButtonStyleName(
-		const ITweenFunctionContainer& InTweenFunctions, const TSharedRef<ITweenFunction>& InFunction
+		const ITweenFunctionContainer& InTweenFunctions, const ITweenFunction& InFunction
 		)
 	{
 		return InTweenFunctions.GetFunctionIdentifier(InFunction);
@@ -64,14 +64,14 @@ FTweenToolbarController::FTweenToolbarController(
 	, OverrideStyle(MakePimpl<FOvershootButtonStyleHack>(*InTweenFunctions))
 {
 	// Function changing commands
-	TweenFunctions->ForEachFunction([this](const TSharedRef<ITweenFunction>& InFunction)
+	TweenFunctions->ForEachFunction([this](const ITweenFunction& InFunction)
 	{
 		CommandList->MapAction(TweenFunctions->GetCommandForFunction(
 			InFunction),
 			FUIAction(
-			   FExecuteAction::CreateRaw(this, &FTweenToolbarController::SetBlendFunction, InFunction.ToSharedPtr()),
+			   FExecuteAction::CreateRaw(this, &FTweenToolbarController::SetBlendFunction, &InFunction),
 			   {},
-			   FIsActionChecked::CreateRaw(this, &FTweenToolbarController::IsBlendFunctionSelected, InFunction.ToSharedPtr())
+			   FIsActionChecked::CreateRaw(this, &FTweenToolbarController::IsBlendFunctionSelected, &InFunction)
 			)
 		);
 	});
@@ -88,7 +88,7 @@ FTweenToolbarController::FTweenToolbarController(
 FTweenToolbarController::~FTweenToolbarController()
 {
 	// Function changing commands
-	TweenFunctions->ForEachFunction([this](const TSharedRef<ITweenFunction>& InFunction)
+	TweenFunctions->ForEachFunction([this](const ITweenFunction& InFunction)
 	{
 		CommandList->UnmapAction(TweenFunctions->GetCommandForFunction(InFunction));
 	});
@@ -125,7 +125,7 @@ TSharedRef<SWidget> FTweenToolbarController::MakeTweenModeMenu() const
 {
 	FMenuBuilder MenuBuilder(true, CommandList);
 
-	TweenFunctions->ForEachFunction([this, &MenuBuilder](const TSharedRef<ITweenFunction>& InFunction)
+	TweenFunctions->ForEachFunction([this, &MenuBuilder](const ITweenFunction& InFunction)
 	{
 		MenuBuilder.AddMenuEntry(TweenFunctions->GetCommandForFunction(InFunction));
 	});
@@ -148,16 +148,15 @@ FSlateIcon FTweenToolbarController::GetIconForComboBox() const
 	return TweenFunctions->GetCommandForFunction(TweenModel->GetTweenFunction())->GetIcon();
 }
 
-void FTweenToolbarController::SetBlendFunction(TSharedPtr<ITweenFunction> InFunction) const
+void FTweenToolbarController::SetBlendFunction(const ITweenFunction* InFunction) const
 {
-	const TSharedRef<ITweenFunction> FunctionRef = InFunction.ToSharedRef();
-	TweenModel->SetTweenFunction(FunctionRef);
-	OnTweenFunctionChangedDelegate.Broadcast(FunctionRef);
+	TweenModel->SetTweenFunction(*InFunction);
+	OnTweenFunctionChangedDelegate.Broadcast(*InFunction);
 }
 
-bool FTweenToolbarController::IsBlendFunctionSelected(TSharedPtr<ITweenFunction> InFunction) const
+bool FTweenToolbarController::IsBlendFunctionSelected(const ITweenFunction* InFunction) const
 {
-	return TweenModel->GetTweenFunction() == InFunction;
+	return &TweenModel->GetTweenFunction() == InFunction;
 }
 
 void FTweenToolbarController::ToggleOvershootMode() const
diff --git a/Engine/Plugins/Animation/TweeningUtils/Source/TweeningUtilsEditor/Public/Math/Abstraction/BuiltInTweenFunctionContainer.h b/Engine/Plugins/Animation/TweeningUtils/Source/TweeningUtilsEditor/Public/Math/Abstraction/BuiltInTweenFunctionContainer.h
--- a/Engine/Plugins/Animation/TweeningUtils/Source/TweeningUtilsEditor/Public/Math/Abstraction/BuiltInTweenFunctionContainer.h
+++ b/Engine/Plugins/Animation/TweeningUtils/Source/TweeningUtilsEditor/Public/Math/Abstraction/BuiltInTweenFunctionContainer.h
@@ -2,15 +2,14 @@
 
 #pragma once
 
+#include "BuiltInTweenFunction.h"
 #include "Containers/Array.h"
 #include "ITweenFunctionContainer.h"
+#include "KeyBlendingAbstraction.h"
 #include "Templates/SharedPointer.h"
 
 namespace UE::TweeningUtilsEditor
 {
-class FBuiltInTweenFunction;
-enum class EBlendFunction : uint8;
-
 /** Contains all the functions from this module, i.e. those specified by EBlendFunction.  */
 class TWEENINGUTILSEDITOR_API FBuiltInTweenFunctionContainer : public ITweenFunctionContainer
 {
@@ -19,26 +18,28 @@ public:
 	FBuiltInTweenFunctionContainer();
 
 	/** @return Gets the type of blend function the tween function corresponds to, or unset if you pass in a function not owned by this container. */
-	TOptional<EBlendFunction> GetBlendFunctionType(const TSharedRef<ITweenFunction>& InFunctionToFind) const;
+	TOptional<EBlendFunction> GetBlendFunctionType(const ITweenFunction& InFunctionToFind) const;
 
 	/** @return The blend function identified by the enum. */
-	TSharedPtr<ITweenFunction> GetBlendFunction(EBlendFunction InBlendFunction) const;
+	const ITweenFunction* GetBlendFunction(EBlendFunction InBlendFunction) const;
 
 	//~ Begin ITweenFunctionContainer Interface
-	virtual void ForEachFunction(const TFunctionRef<void(const TSharedRef<ITweenFunction>&)> InConsumer) const override;
-	virtual TSharedPtr<ITweenFunction> GetFunction(int32 InIndex) const override;
+	virtual void ForEachFunction(const TFunctionRef<void(const ITweenFunction&)> InConsumer) const override;
+	virtual const ITweenFunction* GetFunction(int32 InIndex) const override;
 	virtual int32 NumFunctions() const override { return Functions.Num(); }
-	virtual TSharedPtr<FUICommandInfo> GetCommandForFunction(const TSharedRef<ITweenFunction>& InFunction) const override;
-	virtual const FSlateBrush* GetIconForFunction(const TSharedRef<ITweenFunction>& InFunction) const override;
-	virtual FLinearColor GetColorForFunction(const TSharedRef<ITweenFunction>& InFunction) const override;
-	virtual FText GetLabelForFunction(const TSharedRef<ITweenFunction>& InFunction) const override;
-	virtual FText GetToolTipForFunction(const TSharedRef<ITweenFunction>& InFunction) const override;
-	virtual FString GetFunctionIdentifier(const TSharedRef<ITweenFunction>& InFunction) const override;
+	virtual TSharedPtr<FUICommandInfo> GetCommandForFunction(const ITweenFunction& InFunction) const override;
+	virtual const FSlateBrush* GetIconForFunction(const ITweenFunction& InFunction) const override;
+	virtual FLinearColor GetColorForFunction(const ITweenFunction& InFunction) const override;
+	virtual FText GetLabelForFunction(const ITweenFunction& InFunction) const override;
+	virtual FText GetToolTipForFunction(const ITweenFunction& InFunction) const override;
+	virtual FString GetFunctionIdentifier(const ITweenFunction& InFunction) const override;
 	//~ Begin ITweenFunctionContainer Interface
 
 private:
 
+	using FTweenArray = TArray<FBuiltInTweenFunction, TInlineAllocator<static_cast<int32>(EBlendFunction::Num)>>;
+
 	/** The built-in functions. */
-	TArray<TSharedRef<FBuiltInTweenFunction>> Functions;
+	const FTweenArray Functions;
 };
 }
diff --git a/Engine/Plugins/Animation/TweeningUtils/Source/TweeningUtilsEditor/Public/Math/Abstraction/ITweenFunctionContainer.h b/Engine/Plugins/Animation/TweeningUtils/Source/TweeningUtilsEditor/Public/Math/Abstraction/ITweenFunctionContainer.h
--- a/Engine/Plugins/Animation/TweeningUtils/Source/TweeningUtilsEditor/Public/Math/Abstraction/ITweenFunctionContainer.h
+++ b/Engine/Plugins/Animation/TweeningUtils/Source/TweeningUtilsEditor/Public/Math/Abstraction/ITweenFunctionContainer.h
@@ -21,34 +21,34 @@ class ITweenFunctionContainer
 public:
 
 	/** Iterates through all the functions */
-	virtual void ForEachFunction(const TFunctionRef<void(const TSharedRef<ITweenFunction>&)> InConsumer) const = 0;
+	virtual void ForEachFunction(const TFunctionRef<void(const ITweenFunction&)> InConsumer) const = 0;
 
 	/** @return The function at InIndex. */
-	virtual TSharedPtr<ITweenFunction> GetFunction(int32 InIndex) const = 0;
+	virtual const ITweenFunction* GetFunction(int32 InIndex) const = 0;
 	/** @return The number of functions contained. */
 	virtual int32 NumFunctions() const = 0;
 	/** @return The index of InFunction if found and INDEX_NONE if not. */
-	TWEENINGUTILSEDITOR_API int32 IndexOf(const TSharedRef<ITweenFunction>& InFunction) const;
+	TWEENINGUTILSEDITOR_API int32 IndexOf(const ITweenFunction& InFunction) const;
 	/** @return Whether InIndex is a valid argument for GetFunction. */
 	bool IsValidIndex(int32 InIndex) const { return InIndex >= 0 && InIndex < NumFunctions(); }
 
 	/** @return The command that is used to select the function */
-	virtual TSharedPtr<FUICommandInfo> GetCommandForFunction(const TSharedRef<ITweenFunction>& InFunction) const = 0;
+	virtual TSharedPtr<FUICommandInfo> GetCommandForFunction(const ITweenFunction& InFunction) const = 0;
 
 	/** @return The un-tinted icon to display the function in the UI with */
-	virtual const FSlateBrush* GetIconForFunction(const TSharedRef<ITweenFunction>& InFunction) const = 0;
+	virtual const FSlateBrush* GetIconForFunction(const ITweenFunction& InFunction) const = 0;
 	/** @return The color that represents the function in the UI */
-	virtual FLinearColor GetColorForFunction(const TSharedRef<ITweenFunction>& InFunction) const = 0;
+	virtual FLinearColor GetColorForFunction(const ITweenFunction& InFunction) const = 0;
 
 	/** @return The label to display in the combo button */
-	virtual FText GetLabelForFunction(const TSharedRef<ITweenFunction>& InFunction) const = 0;
+	virtual FText GetLabelForFunction(const ITweenFunction& InFunction) const = 0;
 	/** @return The description to displayed, e.g. in tool tips. */
-	virtual FText GetToolTipForFunction(const TSharedRef<ITweenFunction>& InFunction) const = 0;
+	virtual FText GetToolTipForFunction(const ITweenFunction& InFunction) const = 0;
 
 	/** @return Identifier that uniquely identifies this function. Used e.g. to encode the function type in a config file. */
-	virtual FString GetFunctionIdentifier(const TSharedRef<ITweenFunction>& InFunction) const = 0;
+	virtual FString GetFunctionIdentifier(const ITweenFunction& InFunction) const = 0;
 	/** @return The function identified by the InIdentifier (as returned by a call to GetFunctionIdentifier). */
-	TWEENINGUTILSEDITOR_API TSharedPtr<ITweenFunction> FindFunctionByIdentifier(const FString& InIdentifier) const;
+	TWEENINGUTILSEDITOR_API const ITweenFunction* FindFunctionByIdentifier(const FString& InIdentifier) const;
 
 	virtual ~ITweenFunctionContainer() = default;
 };
diff --git a/Engine/Plugins/Animation/TweeningUtils/Source/TweeningUtilsEditor/Public/Widgets/MVC/Models/CurveTweenModel.h b/Engine/Plugins/Animation/TweeningUtils/Source/TweeningUtilsEditor/Public/Widgets/MVC/Models/CurveTweenModel.h
--- a/Engine/Plugins/Animation/TweeningUtils/Source/TweeningUtilsEditor/Public/Widgets/MVC/Models/CurveTweenModel.h
+++ b/Engine/Plugins/Animation/TweeningUtils/Source/TweeningUtilsEditor/Public/Widgets/MVC/Models/CurveTweenModel.h
@@ -15,7 +15,7 @@ class TWEENINGUTILSEDITOR_API FCurveTweenModel : public FFunctionBasedTweenModel
 {
 public:
 
-	explicit FCurveTweenModel(TWeakPtr<FCurveEditor> InWeakCurveEditor, const TSharedRef<ITweenFunction>& InTweenFunction);
+	explicit FCurveTweenModel(TWeakPtr<FCurveEditor> InWeakCurveEditor, const ITweenFunction& InTweenFunction);
 
 	/** @return Whether it makes sense to call StartBlendOperation, etc. */
 	bool HasAnythingToBlend() const;
diff --git a/Engine/Plugins/Animation/TweeningUtils/Source/TweeningUtilsEditor/Public/Widgets/MVC/Models/FunctionBasedTweenModel.h b/Engine/Plugins/Animation/TweeningUtils/Source/TweeningUtilsEditor/Public/Widgets/MVC/Models/FunctionBasedTweenModel.h
--- a/Engine/Plugins/Animation/TweeningUtils/Source/TweeningUtilsEditor/Public/Widgets/MVC/Models/FunctionBasedTweenModel.h
+++ b/Engine/Plugins/Animation/TweeningUtils/Source/TweeningUtilsEditor/Public/Widgets/MVC/Models/FunctionBasedTweenModel.h
@@ -2,8 +2,7 @@
 
 #pragma once
 
-#include "Math/ContiguousKeyMapping.h"
-#include "Templates/SharedPointer.h"
+#include "HAL/Platform.h"
 #include "TweenModel.h"
 
 class FCurveEditor;
@@ -17,15 +16,16 @@ class FFunctionBasedTweenModel : public FTweenModel
 {
 public:
 
-	explicit FFunctionBasedTweenModel(const TSharedRef<ITweenFunction>& InInitialFunction) : TweenFunction(InInitialFunction) {}
+	explicit FFunctionBasedTweenModel(const ITweenFunction& InInitialFunction UE_LIFETIMEBOUND) : TweenFunction(&InInitialFunction) {}
 
-	virtual void SetTweenFunction(const TSharedRef<ITweenFunction>& InTweenFunction) { TweenFunction = InTweenFunction; }
-	const TSharedRef<ITweenFunction>& GetTweenFunction() const { return TweenFunction; }
+	/** Sets the tween function. You are responsible to make sure it outlives this model. */
+	virtual void SetTweenFunction(const ITweenFunction& InTweenFunction UE_LIFETIMEBOUND) { TweenFunction = &InTweenFunction; }
+	const ITweenFunction& GetTweenFunction() const { return *TweenFunction; }
 	
 private:
 
-	/** The tween function the sub-class will use for blending. */
-	TSharedRef<ITweenFunction> TweenFunction;
+	/** The tween function the sub-class will use for blending. You are responsible to make sure it outlives this model. */
+	const ITweenFunction* TweenFunction;
 };
 }
 
diff --git a/Engine/Plugins/Animation/TweeningUtils/Source/TweeningUtilsEditor/Public/Widgets/MVC/TweenToolbarController.h b/Engine/Plugins/Animation/TweeningUtils/Source/TweeningUtilsEditor/Public/Widgets/MVC/TweenToolbarController.h
--- a/Engine/Plugins/Animation/TweeningUtils/Source/TweeningUtilsEditor/Public/Widgets/MVC/TweenToolbarController.h
+++ b/Engine/Plugins/Animation/TweeningUtils/Source/TweeningUtilsEditor/Public/Widgets/MVC/TweenToolbarController.h
@@ -63,7 +63,7 @@ public:
 	 */
 	void AddToToolbar(FToolBarBuilder& ToolbarBuilder, FMakeWidgetArgs InArgs = {}) const;
 
-	DECLARE_MULTICAST_DELEGATE_OneParam(FOnTweenFunctionChanged, const TSharedRef<ITweenFunction>&);
+	DECLARE_MULTICAST_DELEGATE_OneParam(FOnTweenFunctionChanged, const ITweenFunction&);
 	/** Invoked when the tween function is changed by this controller. */
 	FOnTweenFunctionChanged& OnTweenFunctionChanged() { return OnTweenFunctionChangedDelegate; }
 
@@ -88,8 +88,8 @@ private:
 	FText GetToolTipForComboBox() const;
 	FSlateIcon GetIconForComboBox() const;
 	
-	void SetBlendFunction(TSharedPtr<ITweenFunction> InFunction) const;
-	bool IsBlendFunctionSelected(TSharedPtr<ITweenFunction> InFunction) const;
+	void SetBlendFunction(const ITweenFunction* InFunction) const;
+	bool IsBlendFunctionSelected(const ITweenFunction* InFunction) const;
 	
 	void ToggleOvershootMode() const;
 	bool IsOvershootModeEnabled() const;
diff --git a/Engine/Plugins/Editor/CurveEditorTools/Source/CurveEditorTools/Private/Extensions/Tweening/ContextAwareTweeningModel.cpp b/Engine/Plugins/Editor/CurveEditorTools/Source/CurveEditorTools/Private/Extensions/Tweening/ContextAwareTweeningModel.cpp
--- a/Engine/Plugins/Editor/CurveEditorTools/Source/CurveEditorTools/Private/Extensions/Tweening/ContextAwareTweeningModel.cpp
+++ b/Engine/Plugins/Editor/CurveEditorTools/Source/CurveEditorTools/Private/Extensions/Tweening/ContextAwareTweeningModel.cpp
@@ -8,7 +8,7 @@ namespace UE::CurveEditorTools
 {
 FContextAwareTweeningModel::FContextAwareTweeningModel(
 	TWeakPtr<FCurveEditor> InWeakCurveEditor,
-	const TSharedRef<TweeningUtilsEditor::ITweenFunction>& InTweenFunction
+	const TweeningUtilsEditor::ITweenFunction& InTweenFunction
 	)
 	: FFunctionBasedTweenModel(InTweenFunction)
 	, CurveEditor(InWeakCurveEditor)
@@ -26,7 +26,7 @@ void FContextAwareTweeningModel::SetScaleMode(TweeningUtilsEditor::ETweenScaleMo
 	}
 }
 
-void FContextAwareTweeningModel::SetTweenFunction(const TSharedRef<TweeningUtilsEditor::ITweenFunction>& InTweenFunction)
+void FContextAwareTweeningModel::SetTweenFunction(const TweeningUtilsEditor::ITweenFunction& InTweenFunction)
 {
 	Super::SetTweenFunction(InTweenFunction);
 	SelectionBasedModel.SetTweenFunction(InTweenFunction);
diff --git a/Engine/Plugins/Editor/CurveEditorTools/Source/CurveEditorTools/Private/Extensions/Tweening/ContextAwareTweeningModel.h b/Engine/Plugins/Editor/CurveEditorTools/Source/CurveEditorTools/Private/Extensions/Tweening/ContextAwareTweeningModel.h
--- a/Engine/Plugins/Editor/CurveEditorTools/Source/CurveEditorTools/Private/Extensions/Tweening/ContextAwareTweeningModel.h
+++ b/Engine/Plugins/Editor/CurveEditorTools/Source/CurveEditorTools/Private/Extensions/Tweening/ContextAwareTweeningModel.h
@@ -30,12 +30,12 @@ public:
 	 */
 	explicit FContextAwareTweeningModel(
 		TWeakPtr<FCurveEditor> InWeakCurveEditor,
-		const TSharedRef<TweeningUtilsEditor::ITweenFunction>& InInitialTweenFunction
+		const TweeningUtilsEditor::ITweenFunction& InInitialTweenFunction
 		);
 
 	//~ Begin FTweenModel Interface
 	virtual void SetScaleMode(TweeningUtilsEditor::ETweenScaleMode InMode) override;
-	virtual void SetTweenFunction(const TSharedRef<TweeningUtilsEditor::ITweenFunction>& InTweenFunction) override;
+	virtual void SetTweenFunction(const TweeningUtilsEditor::ITweenFunction& InTweenFunction) override;
 	virtual void StartBlendOperation() override;
 	virtual void StopBlendOperation() override;
 	virtual void BlendValues(float InNormalizedValue) override;
diff --git a/Engine/Plugins/Editor/CurveEditorTools/Source/CurveEditorTools/Private/Extensions/Tweening/KeyInsertingTweenModel.cpp b/Engine/Plugins/Editor/CurveEditorTools/Source/CurveEditorTools/Private/Extensions/Tweening/KeyInsertingTweenModel.cpp
--- a/Engine/Plugins/Editor/CurveEditorTools/Source/CurveEditorTools/Private/Extensions/Tweening/KeyInsertingTweenModel.cpp
+++ b/Engine/Plugins/Editor/CurveEditorTools/Source/CurveEditorTools/Private/Extensions/Tweening/KeyInsertingTweenModel.cpp
@@ -9,7 +9,7 @@
 namespace UE::CurveEditorTools
 {
 FKeyInsertingTweenModel::FKeyInsertingTweenModel(
-	TWeakPtr<FCurveEditor> InWeakCurveEditor, const TSharedRef<TweeningUtilsEditor::ITweenFunction>& InInitialFunction
+	TWeakPtr<FCurveEditor> InWeakCurveEditor, const TweeningUtilsEditor::ITweenFunction& InInitialFunction
 	)
 	: FCurveTweenModel(MoveTemp(InWeakCurveEditor), InInitialFunction)
 {}
diff --git a/Engine/Plugins/Editor/CurveEditorTools/Source/CurveEditorTools/Private/Extensions/Tweening/KeyInsertingTweenModel.h b/Engine/Plugins/Editor/CurveEditorTools/Source/CurveEditorTools/Private/Extensions/Tweening/KeyInsertingTweenModel.h
--- a/Engine/Plugins/Editor/CurveEditorTools/Source/CurveEditorTools/Private/Extensions/Tweening/KeyInsertingTweenModel.h
+++ b/Engine/Plugins/Editor/CurveEditorTools/Source/CurveEditorTools/Private/Extensions/Tweening/KeyInsertingTweenModel.h
@@ -16,7 +16,7 @@ public:
 	
 	explicit FKeyInsertingTweenModel(
 		TWeakPtr<FCurveEditor> InWeakCurveEditor,
-		const TSharedRef<TweeningUtilsEditor::ITweenFunction>& InInitialFunction
+		const TweeningUtilsEditor::ITweenFunction& InInitialFunction
 		);
 
 	//~ Begin FTweenModel
diff --git a/Engine/Plugins/Editor/CurveEditorTools/Source/CurveEditorTools/Private/Extensions/Tweening/TweenEditorExtension.cpp b/Engine/Plugins/Editor/CurveEditorTools/Source/CurveEditorTools/Private/Extensions/Tweening/TweenEditorExtension.cpp
--- a/Engine/Plugins/Editor/CurveEditorTools/Source/CurveEditorTools/Private/Extensions/Tweening/TweenEditorExtension.cpp
+++ b/Engine/Plugins/Editor/CurveEditorTools/Source/CurveEditorTools/Private/Extensions/Tweening/TweenEditorExtension.cpp
@@ -64,12 +64,10 @@ bool FTweenEditorExtension::InitControllers(TSharedRef<FUICommandList> InCommand
 	{
 		using namespace TweeningUtilsEditor;
 
-		const TSharedPtr<ITweenFunction> PreferredFunction = BlendFunctionContainer->FindFunctionByIdentifier(
+		const ITweenFunction* PreferredFunction = BlendFunctionContainer->FindFunctionByIdentifier(
 			UCurveEditorToolsUserSettings::Get()->GetPreferredTweenFunction()
 			);
-		const TSharedRef<ITweenFunction> InitialFunction = PreferredFunction
-			? PreferredFunction.ToSharedRef() : BlendFunctionContainer->GetFunction(0).ToSharedRef();
-
+		const ITweenFunction& InitialFunction = PreferredFunction ? *PreferredFunction : *BlendFunctionContainer->GetFunction(0);
 		
 		using FModel = TEditorTweenModel<FContextAwareTweeningModel>;
 		const TSharedRef<FModel> TweenModel = MakeShared<FModel>(
@@ -84,7 +82,7 @@ bool FTweenEditorExtension::InitControllers(TSharedRef<FUICommandList> InCommand
 	return false;
 }
 
-void FTweenEditorExtension::OnTweenFunctionSelected(const TSharedRef<TweeningUtilsEditor::ITweenFunction>& InFunction) const
+void FTweenEditorExtension::OnTweenFunctionSelected(const TweeningUtilsEditor::ITweenFunction& InFunction) const
 {
 	const FString& FunctionConfigId = BlendFunctionContainer->GetFunctionIdentifier(InFunction);
 	UCurveEditorToolsUserSettings::Get()->SetPreferredTweenFunction(FunctionConfigId);
diff --git a/Engine/Plugins/Editor/CurveEditorTools/Source/CurveEditorTools/Private/Extensions/Tweening/TweenEditorExtension.h b/Engine/Plugins/Editor/CurveEditorTools/Source/CurveEditorTools/Private/Extensions/Tweening/TweenEditorExtension.h
--- a/Engine/Plugins/Editor/CurveEditorTools/Source/CurveEditorTools/Private/Extensions/Tweening/TweenEditorExtension.h
+++ b/Engine/Plugins/Editor/CurveEditorTools/Source/CurveEditorTools/Private/Extensions/Tweening/TweenEditorExtension.h
@@ -43,7 +43,7 @@ private:
 	bool InitControllers(TSharedRef<FUICommandList> InCommandList);
 	
 	/** Invoked when the combo button selects a tween function. */
-	void OnTweenFunctionSelected(const TSharedRef<TweeningUtilsEditor::ITweenFunction>& Shared) const;
+	void OnTweenFunctionSelected(const TweeningUtilsEditor::ITweenFunction& InFunction) const;
 };
 }
 
