Control Rig: Tween slider now remembers the last selected tween function

diff --git a/Engine/Plugins/Animation/ControlRig/Source/ControlRigEditor/Private/EditMode/Tween/TweenLogic.cpp b/Engine/Plugins/Animation/ControlRig/Source/ControlRigEditor/Private/EditMode/Tween/TweenLogic.cpp
--- a/Engine/Plugins/Animation/ControlRig/Source/ControlRigEditor/Private/EditMode/Tween/TweenLogic.cpp
+++ b/Engine/Plugins/Animation/ControlRig/Source/ControlRigEditor/Private/EditMode/Tween/TweenLogic.cpp
@@ -20,15 +20,20 @@ static TSharedRef<FControlRigTweenModel> MakeTweenModel(
 	)
 {
 	using FModel = TweeningUtilsEditor::TEditorTweenModel<FControlRigTweenModel>;
+	
+	const FName FeatureName = TEXT("ControlRigTween"); // The key in UTweeningToolsUserSettings from which to look up the preferred function.
+	const TweeningUtilsEditor::ITweenFunction& DefaultFunction = *InFunctions.GetFunction(0); // Use 1st function if no preferred function is in config.
+	
 	TAttribute<TWeakPtr<FCurveEditor>> CurveEditorAttr = TAttribute<TWeakPtr<FCurveEditor>>::CreateLambda([SequencerAttr = InSequencerAttr]
 	{
 		const TSharedPtr<ISequencer> SequencerPin = SequencerAttr.Get().Pin();
 		return GetCurveEditorFromSequencer(SequencerPin);
 	});
-	// TODO UE-240484: Read preferred function from config
+	
 	return MakeShared<FModel>(
-		MoveTemp(CurveEditorAttr),	// TTangentFlatteningTweenProxy arg
-		MoveTemp(InSequencerAttr), InOwningEditMode, *InFunctions.GetFunction(0)
+		FeatureName, InFunctions,										// TUserPreferredTweenProxy
+		MoveTemp(CurveEditorAttr),										// TTangentFlatteningTweenProxy arg
+		MoveTemp(InSequencerAttr), InOwningEditMode, DefaultFunction
 		);
 }
 }
diff --git a/Engine/Plugins/Animation/TweeningUtils/Source/TweeningUtilsEditor/Private/TweeningToolsUserSettings.cpp b/Engine/Plugins/Animation/TweeningUtils/Source/TweeningUtilsEditor/Private/TweeningToolsUserSettings.cpp
new file mode 100644
--- /dev/null
+++ b/Engine/Plugins/Animation/TweeningUtils/Source/TweeningUtilsEditor/Private/TweeningToolsUserSettings.cpp
@@ -0,0 +1,10 @@
+﻿// Copyright Epic Games, Inc. All Rights Reserved.
+
+#include "TweeningToolsUserSettings.h"
+
+#include "UObject/UObjectGlobals.h"
+
+UTweeningToolsUserSettings* UTweeningToolsUserSettings::Get()
+{
+	return GetMutableDefault<UTweeningToolsUserSettings>();
+}
diff --git a/Engine/Plugins/Animation/TweeningUtils/Source/TweeningUtilsEditor/Public/TweeningToolsUserSettings.h b/Engine/Plugins/Animation/TweeningUtils/Source/TweeningUtilsEditor/Public/TweeningToolsUserSettings.h
new file mode 100644
--- /dev/null
+++ b/Engine/Plugins/Animation/TweeningUtils/Source/TweeningUtilsEditor/Public/TweeningToolsUserSettings.h
@@ -0,0 +1,37 @@
+﻿// Copyright Epic Games, Inc. All Rights Reserved.
+
+#pragma once
+
+#include "Containers/Map.h"
+#include "Containers/UnrealString.h"
+#include "UObject/Object.h"
+#include "UObject/NameTypes.h"
+#include "TweeningToolsUserSettings.generated.h"
+
+UCLASS(Config=EditorPerProjectUserSettings)
+class TWEENINGUTILSEDITOR_API UTweeningToolsUserSettings : public UObject
+{
+	GENERATED_BODY()
+public:
+
+	static UTweeningToolsUserSettings* Get();
+
+	/** Sets the preferred tween function for the given feature. */
+	void SetPreferredTweenFunction(FName InFeatureKey, FString InNewPreferredFunction)
+	{
+		PreferredTweenFunction.Add(InFeatureKey, MoveTemp(InNewPreferredFunction));
+		SaveConfig();
+	}
+	/** Gets the preferred tween function for the given feature, if there is any. */
+	const FString* GetPreferredTweenFunction(FName InFeatureKey) const { return PreferredTweenFunction.Find(InFeatureKey); }
+	
+private:
+	
+	/**
+	 * Associates features with the preferred function in that feature.
+	 * If you want multiple locations in the editor to share the same setting, use the same key.
+	 * This tween function should be selected by default when a curve editor is created. It is the function that was used last time in the editor.
+	 */
+	UPROPERTY(Config)
+	TMap<FName, FString> PreferredTweenFunction;
+};
diff --git a/Engine/Plugins/Animation/TweeningUtils/Source/TweeningUtilsEditor/Public/Widgets/MVC/Models/EditorTweenModel.h b/Engine/Plugins/Animation/TweeningUtils/Source/TweeningUtilsEditor/Public/Widgets/MVC/Models/EditorTweenModel.h
--- a/Engine/Plugins/Animation/TweeningUtils/Source/TweeningUtilsEditor/Public/Widgets/MVC/Models/EditorTweenModel.h
+++ b/Engine/Plugins/Animation/TweeningUtils/Source/TweeningUtilsEditor/Public/Widgets/MVC/Models/EditorTweenModel.h
@@ -4,24 +4,28 @@
 
 #include "TangentFlatteningTweenProxy.h"
 #include "TransactionalTweenModelProxy.h"
+#include "UserPreferredTweenProxy.h"
+
+#include <type_traits>
 
 namespace UE::TweeningUtilsEditor
 {
 /**
  * Extends TUnderlyingModel with commonly expected shared functionality:
  *	- Undo / redo
+ *	- Keeps track of the last selected function and saves it in UTweeningToolsUserSettings.
  *	- As the tween function squishes the keys vertically, user-specified tangents are flattened.
  *	- TODO UE-234007: Synchronize tween function selection from different editor tools across all instances of the tweener.
- *	- TODO UE-240484: Move preferred function editor UCurveEditorToolsUserSettings to this module.
  *	- In the future this can be extended to remove vertically stacked keys (same keys per frame). For now, no implemented tween function stacks keys.
  *
  * Construct like this:
  *  using FModel = TEditorTweenModel<FFooTweenModel> // FFooTweenModel should inherit from FTweenModel.
  *	const TSharedRef<FModel> TweenModel = MakeShared<FModel>(
- *		WeakCurveEditor,					// TTangentFlatteningTweenProxy arg
- *		Bar									// Any arguments that FFooTweenModel requires for construction
+ *		TEXT("YourFeatureName"), FunctionContainer	// TUserPreferredTweenProxy
+ *		WeakCurveEditor,							// TTangentFlatteningTweenProxy arg
+ *		Bar											// Any arguments that FFooTweenModel requires for construction
  *	);
  */
-template<typename TUnderlyingModel>
-using TEditorTweenModel = TTransactionalTweenModelProxy<TTangentFlatteningTweenProxy<TUnderlyingModel>>;
+template<typename TUnderlyingModel> requires std::is_base_of_v<FFunctionBasedTweenModel, TUnderlyingModel> // required by TUserPreferredTweenProxy.
+using TEditorTweenModel = TTransactionalTweenModelProxy<TUserPreferredTweenProxy<TTangentFlatteningTweenProxy<TUnderlyingModel>>>;
 }
diff --git a/Engine/Plugins/Animation/TweeningUtils/Source/TweeningUtilsEditor/Public/Widgets/MVC/Models/UserPreferredTweenProxy.h b/Engine/Plugins/Animation/TweeningUtils/Source/TweeningUtilsEditor/Public/Widgets/MVC/Models/UserPreferredTweenProxy.h
new file mode 100644
--- /dev/null
+++ b/Engine/Plugins/Animation/TweeningUtils/Source/TweeningUtilsEditor/Public/Widgets/MVC/Models/UserPreferredTweenProxy.h
@@ -0,0 +1,56 @@
+﻿// Copyright Epic Games, Inc. All Rights Reserved.
+
+#pragma once
+
+#include "FunctionBasedTweenModel.h"
+#include "Math/Abstraction/ITweenFunctionContainer.h"
+#include "UObject/NameTypes.h"
+#include "TweeningToolsUserSettings.h"
+#include <type_traits>
+
+namespace UE::TweeningUtilsEditor
+{
+/** Updates the user config with the last selected tween function whenever it is changed.  */
+template<typename TBase> requires std::is_base_of_v<FFunctionBasedTweenModel, TBase>
+class TUserPreferredTweenProxy : public TBase
+{
+public:
+
+	template<typename... TArg>
+	explicit TUserPreferredTweenProxy(FName InFeatureName, const ITweenFunctionContainer& InOwningFunctionContainer UE_LIFETIMEBOUND, TArg&&... Arg)
+		: TBase(Forward<TArg>(Arg)...)
+		, FeatureName(InFeatureName)
+		, OwningFunctionContainer(InOwningFunctionContainer)
+	{
+		// If there's a preferred function saved, load that up as default.
+		const FString* PreferredFunctionId = UTweeningToolsUserSettings::Get()->GetPreferredTweenFunction(FeatureName);
+		const ITweenFunction* PreferredFunction = PreferredFunctionId
+			? OwningFunctionContainer.FindFunctionByIdentifier(*PreferredFunctionId) : nullptr;
+		if (PreferredFunction)
+		{
+			TBase::SetTweenFunction(*PreferredFunction);
+		}
+	}
+
+	//~ Begin FTweenModel Interface
+	virtual void SetTweenFunction(const ITweenFunction& InTweenFunction UE_LIFETIMEBOUND)
+	{
+		const bool bIsContained = OwningFunctionContainer.IsValidIndex(OwningFunctionContainer.IndexOf(InTweenFunction));
+		FString FunctionId = OwningFunctionContainer.GetFunctionIdentifier(InTweenFunction);
+		if (ensureMsgf(bIsContained, TEXT("InTweenFunction is not in ITweenFunctionContainer!")) && !FunctionId.IsEmpty())
+		{
+			UTweeningToolsUserSettings::Get()->SetPreferredTweenFunction(FeatureName, MoveTemp(FunctionId));
+		}
+		
+		TBase::SetTweenFunction(InTweenFunction);
+	}
+	//~ Begin FTweenModel Interface
+
+private:
+
+	/** The key under which the tween function will be saved. */
+	const FName FeatureName;
+	/** The function container from which InTweenFunction originates. You're supposed to ensure that the function container outlives this object.*/
+	const ITweenFunctionContainer& OwningFunctionContainer;
+};
+}
diff --git a/Engine/Plugins/Editor/CurveEditorTools/Source/CurveEditorTools/Private/CurveEditorToolsUserSettings.cpp b/Engine/Plugins/Editor/CurveEditorTools/Source/CurveEditorTools/Private/CurveEditorToolsUserSettings.cpp
deleted file mode 100644
diff --git a/Engine/Plugins/Editor/CurveEditorTools/Source/CurveEditorTools/Private/CurveEditorToolsUserSettings.h b/Engine/Plugins/Editor/CurveEditorTools/Source/CurveEditorTools/Private/CurveEditorToolsUserSettings.h
deleted file mode 100644
diff --git a/Engine/Plugins/Editor/CurveEditorTools/Source/CurveEditorTools/Private/Extensions/Tweening/TweenEditorExtension.cpp b/Engine/Plugins/Editor/CurveEditorTools/Source/CurveEditorTools/Private/Extensions/Tweening/TweenEditorExtension.cpp
--- a/Engine/Plugins/Editor/CurveEditorTools/Source/CurveEditorTools/Private/Extensions/Tweening/TweenEditorExtension.cpp
+++ b/Engine/Plugins/Editor/CurveEditorTools/Source/CurveEditorTools/Private/Extensions/Tweening/TweenEditorExtension.cpp
@@ -4,7 +4,6 @@
 
 #include "ContextAwareTweeningModel.h"
 #include "CurveEditor.h"
-#include "CurveEditorToolsUserSettings.h"
 #include "Framework/MultiBox/MultiBoxBuilder.h"
 #include "Framework/MultiBox/MultiBoxExtender.h"
 #include "Math/Abstraction/BuiltInTweenFunctionContainer.h"
@@ -63,28 +62,20 @@ bool FTweenEditorExtension::InitControllers(TSharedRef<FUICommandList> InCommand
 	if (const TSharedPtr<FCurveEditor> CurveEditorPin = WeakCurveEditor.Pin(); ensure(CurveEditorPin))
 	{
 		using namespace TweeningUtilsEditor;
-
-		const ITweenFunction* PreferredFunction = BlendFunctionContainer->FindFunctionByIdentifier(
-			UCurveEditorToolsUserSettings::Get()->GetPreferredTweenFunction()
-			);
-		const ITweenFunction& InitialFunction = PreferredFunction ? *PreferredFunction : *BlendFunctionContainer->GetFunction(0);
+		
+		const FName FeatureName = TEXT("CurveEditorTween"); // The key in UTweeningToolsUserSettings from which to look up the preferred function.
+		const ITweenFunction& DefaultFunction = *BlendFunctionContainer->GetFunction(0); // Use 1st function if no preferred function is in config.
 		
 		using FModel = TEditorTweenModel<FContextAwareTweeningModel>;
 		const TSharedRef<FModel> TweenModel = MakeShared<FModel>(
-			WeakCurveEditor,					// TTangentFlatteningTweenProxy arg
-			WeakCurveEditor, InitialFunction	// FContextAwareTweeningModel args
+			FeatureName, *BlendFunctionContainer,	// TUserPreferredTweenProxy
+			WeakCurveEditor,							// TTangentFlatteningTweenProxy arg
+			WeakCurveEditor, DefaultFunction			// FContextAwareTweeningModel args
 			);
 		TweenControllers.Emplace(TweenModel, InCommandList, BlendFunctionContainer);
-		TweenControllers->ToolbarController.OnTweenFunctionChanged().AddRaw(this, &FTweenEditorExtension::OnTweenFunctionSelected);
 		
 		return true;
 	}
 	return false;
 }
-
-void FTweenEditorExtension::OnTweenFunctionSelected(const TweeningUtilsEditor::ITweenFunction& InFunction) const
-{
-	const FString& FunctionConfigId = BlendFunctionContainer->GetFunctionIdentifier(InFunction);
-	UCurveEditorToolsUserSettings::Get()->SetPreferredTweenFunction(FunctionConfigId);
-}
 }
diff --git a/Engine/Plugins/Editor/CurveEditorTools/Source/CurveEditorTools/Private/Extensions/Tweening/TweenEditorExtension.h b/Engine/Plugins/Editor/CurveEditorTools/Source/CurveEditorTools/Private/Extensions/Tweening/TweenEditorExtension.h
--- a/Engine/Plugins/Editor/CurveEditorTools/Source/CurveEditorTools/Private/Extensions/Tweening/TweenEditorExtension.h
+++ b/Engine/Plugins/Editor/CurveEditorTools/Source/CurveEditorTools/Private/Extensions/Tweening/TweenEditorExtension.h
@@ -41,9 +41,6 @@ private:
 
 	/** Makes sure that ToolbarController is created. */
 	bool InitControllers(TSharedRef<FUICommandList> InCommandList);
-	
-	/** Invoked when the combo button selects a tween function. */
-	void OnTweenFunctionSelected(const TweeningUtilsEditor::ITweenFunction& InFunction) const;
 };
 }
 
