Tweening Utils: Refactored some usages of FKeyPosition with FVector2d. Simplified interface of FKeyBlendInfo:
- Simplified FKeyBlendInfo no longer has array of all keys. This input communicated the wrong idea to the reader (that iterating through all blended keys may be a valid thing for a tween function to do). Now, it is clear that a tween function only tweens based off of info about the current key.
- Removed FKeyPosition so the code can be moved into a runtime module

diff --git a/Engine/Plugins/Animation/ControlRig/Source/ControlRigEditor/Private/EditMode/Tween/ControlRigTweenModel.cpp b/Engine/Plugins/Animation/ControlRig/Source/ControlRigEditor/Private/EditMode/Tween/ControlRigTweenModel.cpp
--- a/Engine/Plugins/Animation/ControlRig/Source/ControlRigEditor/Private/EditMode/Tween/ControlRigTweenModel.cpp
+++ b/Engine/Plugins/Animation/ControlRig/Source/ControlRigEditor/Private/EditMode/Tween/ControlRigTweenModel.cpp
@@ -80,11 +80,6 @@ bool FControlRigTweenModel::DoBlend(float InNormalizedValue)
 		const FFrameTime FrameTime = SequencerPin->GetLocalTime().Time;
 		const FFrameRate TickResolution = SequencerPin->GetFocusedTickResolution();
 		
-		const TArray<FKeyPosition> AllKeyPositions;
-		const TArray <int32> Indices;
-		const int32 CurrentIndex = -1;
-		TweeningUtilsEditor::FKeyBlendInfo BlendStruct(AllKeyPositions, Indices);
-
 		bool bDidBlend = false;
 		for (const FControlRigObjectSelection::FObjectChannels& ObjectChannels : ObjectSelection.ChannelsArray)
 		{
@@ -102,8 +97,19 @@ bool FControlRigTweenModel::DoBlend(float InNormalizedValue)
 					const double NextTime = TickResolution.AsSeconds(FFrameTime(ObjectChannels.KeyBounds[Index].NextFrame));
 					const double CurrentValue = ObjectChannels.KeyBounds[Index].CurrentValue;
 					const double CurrentTime = TickResolution.AsSeconds(FFrameTime(ObjectChannels.KeyBounds[Index].CurrentFrame));
-					BlendStruct.SetValues(PreviousTime, PreviousValue, CurrentTime, CurrentValue,
-						NextTime, NextValue, ScaledBlendValue, CurrentValue, CurrentValue, CurrentIndex);
+
+					const FVector2d BeforeBlendRange { PreviousTime, PreviousValue };
+					const FVector2d AfterBlendRange { NextTime, NextValue };
+					const FVector2d Current { CurrentTime, CurrentValue };
+					
+					// First and last blended are the same as current since we're only blending a single key
+					const TweeningUtilsEditor::FKeyBlendInfo BlendStruct(
+						BeforeBlendRange, Current, // Current is the last key to blend
+						// Key before current, current, and after current key
+						BeforeBlendRange, Current, AfterBlendRange, 
+						Current, AfterBlendRange, // Current is also the last key to blend
+						ScaledBlendValue
+						);
 					const double NewValue = GetTweenFunction().Blend(BlendStruct);
 					using namespace UE::MovieScene;
 					if (ObjectChannels.KeyBounds[Index].FloatChannel)
diff --git a/Engine/Plugins/Animation/TweeningUtils/Source/TweeningUtilsEditor/Private/Math/ContiguousKeyMapping.cpp b/Engine/Plugins/Animation/TweeningUtils/Source/TweeningUtilsEditor/Private/Math/ContiguousKeyMapping.cpp
--- a/Engine/Plugins/Animation/TweeningUtils/Source/TweeningUtilsEditor/Private/Math/ContiguousKeyMapping.cpp
+++ b/Engine/Plugins/Animation/TweeningUtils/Source/TweeningUtilsEditor/Private/Math/ContiguousKeyMapping.cpp
@@ -52,6 +52,11 @@ static void AppendKeyArray(
 		TArray<FKeyPosition> AllKeyPositions;
 		AllKeyPositions.SetNum(AllKeyHandles.Num());
 		Curve->GetKeyPositions(AllKeyHandles, AllKeyPositions);
+		TArray<FVector2d> AllKeyPositionVectors;
+		Algo::Transform(AllKeyPositions, AllKeyPositionVectors, [](const FKeyPosition& InKeyPostition)
+		{
+			return FVector2d{ InKeyPostition.InputValue, InKeyPostition.OutputValue };
+		});
 
 		// Get all the selected keys
 		TArray<FKeyPosition> KeyPositions;
@@ -59,7 +64,7 @@ static void AppendKeyArray(
 		Curve->GetKeyPositions(KeysToBlend, KeyPositions);
 
 		FContiguousKeyMapping::FContiguousKeysArray& KeyArray = OutKeyMap.Add(
-			InCurveId, FContiguousKeyMapping::FContiguousKeysArray(AllKeyHandles, AllKeyPositions)
+			InCurveId, FContiguousKeyMapping::FContiguousKeysArray(AllKeyHandles, AllKeyPositionVectors)
 			);
 
 		TArray<int32> SelectedIndices;
@@ -112,30 +117,17 @@ static TMap<FCurveModelID, FContiguousKeyMapping::FContiguousKeysArray> ComputeK
 }
 
 FContiguousKeyMapping::FContiguousKeys::FContiguousKeys(
-	const TArray<FKeyHandle>& AllKeyHandles, const TArray<FKeyPosition>& AllKeyPositions, const TArray<int32>& InContiguousKeyIndices
+	const TArray<FKeyHandle>& AllKeyHandles, const TArray<FVector2d>& AllKeyPositions, const TArray<int32>& InContiguousKeyIndices
 	)
+	: Indices(InContiguousKeyIndices)
+	, PreviousIndex(AllKeyHandles.IsValidIndex(Indices[0] - 1) ? Indices[0] - 1 : INDEX_NONE)
+	, NextIndex(AllKeyHandles.IsValidIndex(Indices[Indices.Num() -1] + 1) ? Indices[Indices.Num() -1] + 1 : INDEX_NONE)
 {
-	Indices.Reset();
-	if (AllKeyHandles.Num() > 0 && AllKeyPositions.Num() == AllKeyHandles.Num() && InContiguousKeyIndices.Num() > 0)
-	{
-		Indices = InContiguousKeyIndices;
-		//Previous
-		PreviousIndex = Indices[0] - 1;
-		if (PreviousIndex < 0)
-		{
-			PreviousIndex = INDEX_NONE;
-		}		
-		//Next
-		NextIndex = Indices[Indices.Num() -1] + 1;
-		if (NextIndex >= AllKeyHandles.Num())
-		{
-			NextIndex = INDEX_NONE;
-		}
-	}
+	check(AllKeyHandles.Num() > 0 && AllKeyPositions.Num() == AllKeyHandles.Num() && InContiguousKeyIndices.Num() > 0);
 }
 
 FContiguousKeyMapping::FContiguousKeysArray::FContiguousKeysArray(
-	const TArray<FKeyHandle>& InAllKeyHandles, const TArray<FKeyPosition>& InAllKeyPositions
+	const TArray<FKeyHandle>& InAllKeyHandles, const TArray<FVector2d>& InAllKeyPositions
 	)
 	: AllKeyHandles(InAllKeyHandles)
 	, AllKeyPositions(InAllKeyPositions)
diff --git a/Engine/Plugins/Animation/TweeningUtils/Source/TweeningUtilsEditor/Private/Math/KeyBlendingFunctions.cpp b/Engine/Plugins/Animation/TweeningUtils/Source/TweeningUtilsEditor/Private/Math/KeyBlendingFunctions.cpp
--- a/Engine/Plugins/Animation/TweeningUtils/Source/TweeningUtilsEditor/Private/Math/KeyBlendingFunctions.cpp
+++ b/Engine/Plugins/Animation/TweeningUtils/Source/TweeningUtilsEditor/Private/Math/KeyBlendingFunctions.cpp
@@ -3,6 +3,7 @@
 #include "Math/KeyBlendingFunctions.h"
 
 #include "Math/KeyBlendInfo.h"
+#include "Math/Vector2D.h"
 
 namespace UE::TweeningUtilsEditor
 {
@@ -10,55 +11,55 @@ double Blend_ControlsToTween(const FKeyBlendInfo& InBlendInfo)
 {
 	// Classic tween will move all to same location, not based on current time at all just blend and values
 	const double NormalizedBlendValue = (InBlendInfo.BlendValue + 1.0f) * 0.5f;
-	const double Value = InBlendInfo.PreviousValue + (InBlendInfo.NextValue - InBlendInfo.PreviousValue) * (NormalizedBlendValue);
+	const double Value = InBlendInfo.BeforeBlendRange.Y + (InBlendInfo.AfterBlendRange.Y - InBlendInfo.BeforeBlendRange.Y) * (NormalizedBlendValue);
 	return Value;
 }
 
 double Blend_PushPull(const FKeyBlendInfo& InBlendInfo)
 {
-	if (FMath::IsNearlyEqual(InBlendInfo.NextTime, InBlendInfo.PreviousTime))
+	if (FMath::IsNearlyEqual(InBlendInfo.AfterBlendRange.X, InBlendInfo.BeforeBlendRange.X))
 	{
-		return InBlendInfo.CurrentValue;
+		return InBlendInfo.Current.Y;
 	}
-	const double T = (InBlendInfo.CurrentTime - InBlendInfo.PreviousTime) / (InBlendInfo.NextTime - InBlendInfo.PreviousTime);
-	const double ValueAtT = InBlendInfo.PreviousValue + T * (InBlendInfo.NextValue - InBlendInfo.PreviousValue);
+	const double T = (InBlendInfo.Current.X - InBlendInfo.BeforeBlendRange.X) / (InBlendInfo.AfterBlendRange.X - InBlendInfo.BeforeBlendRange.X);
+	const double ValueAtT = InBlendInfo.BeforeBlendRange.Y + T * (InBlendInfo.AfterBlendRange.Y - InBlendInfo.BeforeBlendRange.Y);
 	double NewValue;
 	if (InBlendInfo.BlendValue < 0.0)
 	{
-		NewValue = InBlendInfo.CurrentValue + (-1.0 * InBlendInfo.BlendValue) * (ValueAtT - InBlendInfo.CurrentValue);
+		NewValue = InBlendInfo.Current.Y + (-1.0 * InBlendInfo.BlendValue) * (ValueAtT - InBlendInfo.Current.Y);
 	}
 	else
 	{
-		const double AmplifyValueAtT = InBlendInfo.CurrentValue + (InBlendInfo.CurrentValue - ValueAtT);
-		NewValue = InBlendInfo.CurrentValue + InBlendInfo.BlendValue * (AmplifyValueAtT - InBlendInfo.CurrentValue);
+		const double AmplifyValueAtT = InBlendInfo.Current.Y + (InBlendInfo.Current.Y - ValueAtT);
+		NewValue = InBlendInfo.Current.Y + InBlendInfo.BlendValue * (AmplifyValueAtT - InBlendInfo.Current.Y);
 	}
 	return NewValue;
 }
 
 double Blend_Neighbor(const FKeyBlendInfo& InBlendInfo)
 {
-	double NewValue = InBlendInfo.CurrentValue;
+	double NewValue = InBlendInfo.Current.Y;
 	if (InBlendInfo.BlendValue < 0.0)
 	{
-		NewValue = InBlendInfo.CurrentValue + (-1.0 * InBlendInfo.BlendValue) * (InBlendInfo.PreviousValue - InBlendInfo.CurrentValue);
+		NewValue = InBlendInfo.Current.Y + (-1.0 * InBlendInfo.BlendValue) * (InBlendInfo.BeforeBlendRange.Y - InBlendInfo.Current.Y);
 	}
 	else
 	{
-		NewValue = InBlendInfo.CurrentValue + InBlendInfo.BlendValue * (InBlendInfo.NextValue - InBlendInfo.CurrentValue);
+		NewValue = InBlendInfo.Current.Y + InBlendInfo.BlendValue * (InBlendInfo.AfterBlendRange.Y - InBlendInfo.Current.Y);
 	}
 	return NewValue;
 }
 
 double Blend_Relative(const FKeyBlendInfo& InBlendInfo)
 {
-	double NewValue = InBlendInfo.CurrentValue;
+	double NewValue = InBlendInfo.Current.Y;
 	if (InBlendInfo.BlendValue < 0.0)
 	{
-		NewValue = InBlendInfo.CurrentValue + (-1.0 * InBlendInfo.BlendValue) * (InBlendInfo.PreviousValue - InBlendInfo.FirstValue);
+		NewValue = InBlendInfo.Current.Y + (-1.0 * InBlendInfo.BlendValue) * (InBlendInfo.BeforeBlendRange.Y - InBlendInfo.FirstBlended.Y);
 	}
 	else
 	{
-		NewValue = InBlendInfo.CurrentValue + InBlendInfo.BlendValue * (InBlendInfo.NextValue - InBlendInfo.LastValue);
+		NewValue = InBlendInfo.Current.Y + InBlendInfo.BlendValue * (InBlendInfo.AfterBlendRange.Y - InBlendInfo.LastBlended.Y);
 	}
 	return NewValue;
 }
@@ -83,27 +84,27 @@ static float SCurve(float X, float Slope, float Width, float Height, float XShif
 
 double Blend_Ease(const FKeyBlendInfo& InBlendInfo)
 {
-	if (FMath::IsNearlyEqual(InBlendInfo.NextTime, InBlendInfo.PreviousTime))
+	if (FMath::IsNearlyEqual(InBlendInfo.AfterBlendRange.X, InBlendInfo.BeforeBlendRange.X))
 	{
-		return InBlendInfo.CurrentValue;
+		return InBlendInfo.Current.Y;
 	}
-	const double Source = InBlendInfo.CurrentValue;
-	const double FullTimeDiff = InBlendInfo.NextTime - InBlendInfo.PreviousTime;
+	const double Source = InBlendInfo.Current.Y;
+	const double FullTimeDiff = InBlendInfo.AfterBlendRange.X - InBlendInfo.BeforeBlendRange.X;
 	const double AbsValue = FMath::Abs(InBlendInfo.BlendValue);
-	const double X = InBlendInfo.CurrentTime - InBlendInfo.PreviousTime;
+	const double X = InBlendInfo.Current.X - InBlendInfo.BeforeBlendRange.X;
 	const double Ratio = X / FullTimeDiff;
 	double Shift, Delta, Base;
 	if (InBlendInfo.BlendValue > 0)
 	{
 		Shift = -1.0;
-		Delta = InBlendInfo.NextValue - Source;
+		Delta = InBlendInfo.AfterBlendRange.Y - Source;
 		Base = Source;
 	}
 	else
 	{
 		Shift = 0.0;
-		Delta = Source - InBlendInfo.PreviousValue;
-		Base = InBlendInfo.PreviousValue;
+		Delta = Source - InBlendInfo.BeforeBlendRange.Y;
+		Base = InBlendInfo.BeforeBlendRange.Y;
 	}
 	const double Slope = 5.0 * AbsValue;
 	const double EaseY = Private::BlendToEase::SCurve(Ratio, Slope, 2.0, 2.0, Shift, Shift);
@@ -114,18 +115,17 @@ double Blend_Ease(const FKeyBlendInfo& InBlendInfo)
 
 double Blend_SmoothRough(const FKeyBlendInfo& InBlendInfo)
 {
-	double PrevVal = InBlendInfo.CurrentIndex > 0 ? InBlendInfo.AllKeyPositions[InBlendInfo.Indices[InBlendInfo.CurrentIndex - 1]].OutputValue : InBlendInfo.PreviousValue;
-	double CurVal = InBlendInfo.CurrentValue;
-	double NextVal = InBlendInfo.CurrentIndex < InBlendInfo.Indices.Num() - 1 ? InBlendInfo.AllKeyPositions[InBlendInfo.Indices[InBlendInfo.CurrentIndex + 1]].OutputValue : InBlendInfo.NextValue;
-	double NewValue = (PrevVal * 0.25 + CurVal * 0.5 + NextVal * 0.25);
+	const double PrevVal = InBlendInfo.BeforeCurrent.Y;
+	const double CurVal = InBlendInfo.Current.Y;
+	const double NextVal = InBlendInfo.AfterCurrent.Y;
+	const double NewValue = (PrevVal * 0.25 + CurVal * 0.5 + NextVal * 0.25);
 	if (InBlendInfo.BlendValue < 0.0)
 	{
-		NewValue = InBlendInfo.CurrentValue + (-1.0 * InBlendInfo.BlendValue) * (NewValue - InBlendInfo.CurrentValue);
+		return InBlendInfo.Current.Y + (-1.0 * InBlendInfo.BlendValue) * (NewValue - InBlendInfo.Current.Y);
 	}
 	else
 	{
-		NewValue = InBlendInfo.CurrentValue + InBlendInfo.BlendValue * (InBlendInfo.CurrentValue - NewValue);
+		return InBlendInfo.Current.Y + InBlendInfo.BlendValue * (InBlendInfo.Current.Y - NewValue);
 	}
-	return NewValue;
 }
 }
diff --git a/Engine/Plugins/Animation/TweeningUtils/Source/TweeningUtilsEditor/Public/Math/ContiguousKeyMapping.h b/Engine/Plugins/Animation/TweeningUtils/Source/TweeningUtilsEditor/Public/Math/ContiguousKeyMapping.h
--- a/Engine/Plugins/Animation/TweeningUtils/Source/TweeningUtilsEditor/Public/Math/ContiguousKeyMapping.h
+++ b/Engine/Plugins/Animation/TweeningUtils/Source/TweeningUtilsEditor/Public/Math/ContiguousKeyMapping.h
@@ -18,23 +18,34 @@ struct TWEENINGUTILSEDITOR_API FContiguousKeyMapping
 	// Set of keys in a row
 	struct FContiguousKeys
 	{
-		int32 PreviousIndex;
-		TArray<int32> Indices;
-		int32 NextIndex;
+		/** Indices of keys to blend */
+		const TArray<int32> Indices;
+		/** The index of the key before Indices[0]. INDEX_NONE if there is no previous key. */
+		const int32 PreviousIndex = INDEX_NONE;
+		/** The index of the key after Indidces' last element. INDEX_NONE if there is no next key. */
+		const int32 NextIndex = INDEX_NONE;
+		
 		explicit FContiguousKeys(
-			const TArray<FKeyHandle>& InAllKeyHandles, const TArray<FKeyPosition>& InAllKeyPositions, const TArray<int32>& ContiguousKeyIndices
+			const TArray<FKeyHandle>& InAllKeyHandles, const TArray<FVector2d>& InAllKeyPositions, const TArray<int32>& ContiguousKeyIndices
 			);
 	};
 
 	// All the selected keys in a row for a particular curve and pointers to all of the keys
 	struct FContiguousKeysArray
 	{
-		TArray<FKeyHandle> AllKeyHandles;
-		TArray<FKeyPosition> AllKeyPositions;
-
+		const TArray<FKeyHandle> AllKeyHandles;
+		const TArray<FVector2d> AllKeyPositions;
+
+		/**
+		 * The user selected keys split up into consecutive ranges.
+		 * 
+		 * Let 'x' denote keys selected by the user and '.' keys not selected.
+		 * Suppose the user selection was '.xx...xxx.', i.e. the user selected keys and then shift+selected additional keys somewhere else on the curve.
+		 * Then this array would contain indices of the two ranges 'xx' and 'xxx' is not a blendable range.
+		 */
 		TArray<FContiguousKeys> KeysArray;
 
-		explicit FContiguousKeysArray(const TArray<FKeyHandle>& InAllKeyHandle, const TArray<FKeyPosition>& InAllKeyPositions);
+		explicit FContiguousKeysArray(const TArray<FKeyHandle>& InAllKeyHandle, const TArray<FVector2d>& InAllKeyPositions);
 		void Add(const TArray<int32>& ContiguousKeyIndices);
 	};
 
diff --git a/Engine/Plugins/Animation/TweeningUtils/Source/TweeningUtilsEditor/Public/Math/CurveBlending.h b/Engine/Plugins/Animation/TweeningUtils/Source/TweeningUtilsEditor/Public/Math/CurveBlending.h
--- a/Engine/Plugins/Animation/TweeningUtils/Source/TweeningUtilsEditor/Public/Math/CurveBlending.h
+++ b/Engine/Plugins/Animation/TweeningUtils/Source/TweeningUtilsEditor/Public/Math/CurveBlending.h
@@ -38,43 +38,40 @@ bool BlendCurves(float InBlendValue, const FCurveEditor& InCurveEditor, const FC
 	TArray<FKeyPosition> KeyPositions;
 	for (const TPair<FCurveModelID, FContiguousKeyMapping::FContiguousKeysArray>& KeysArray : InKeySelection.KeyMap)
 	{
-		if (FCurveModel* Curve = InCurveEditor.FindCurve(KeysArray.Key))
+		FCurveModel* Curve = InCurveEditor.FindCurve(KeysArray.Key);
+		if (!Curve)
 		{
-			Curve->Modify();
-			for (const FContiguousKeyMapping::FContiguousKeys& Keys : KeysArray.Value.KeysArray)
-			{
-				const int32 PreviousIndex = Keys.PreviousIndex != INDEX_NONE ? Keys.PreviousIndex : Keys.Indices[0];
-				const double PreviousTime = KeysArray.Value.AllKeyPositions[PreviousIndex].InputValue;
-				const double PreviousValue = KeysArray.Value.AllKeyPositions[PreviousIndex].OutputValue;
-				const int32 NextIndex = Keys.NextIndex != INDEX_NONE ? Keys.NextIndex : Keys.Indices[Keys.Indices.Num() - 1];
-				const double NextTime = KeysArray.Value.AllKeyPositions[NextIndex].InputValue;
-				const double NextValue = KeysArray.Value.AllKeyPositions[NextIndex].OutputValue;
+			continue;
+		}
+		Curve->Modify();
+		
+		const TConstArrayView<FKeyHandle>& AllKeyHandles = KeysArray.Value.AllKeyHandles;
+		const TConstArrayView<FVector2d>& AllKeyPositions = KeysArray.Value.AllKeyPositions;
+		for (const FContiguousKeyMapping::FContiguousKeys& Keys : KeysArray.Value.KeysArray)
+		{
+			const TConstArrayView<int32>& Indices = Keys.Indices;
+				
+			// Get the key before & after the blended keys
+			const int32 BeforeIndex = Keys.PreviousIndex != INDEX_NONE ? Keys.PreviousIndex : Keys.Indices[0];
+			const FVector2d& BeforeBlendRange = AllKeyPositions[BeforeIndex];
+			const int32 AfterIndex = Keys.NextIndex != INDEX_NONE ? Keys.NextIndex : Indices[Indices.Num() - 1];
+			const FVector2d& AfterBlendRange = AllKeyPositions[AfterIndex];
 
-				const int32 NumIndices = Keys.Indices.Num();
-				KeyHandles.SetNum(NumIndices);
-				KeyPositions.SetNum(NumIndices);
-				const double FirstValue = KeysArray.Value.AllKeyPositions[Keys.Indices[0]].OutputValue;
-				const double LastValue = KeysArray.Value.AllKeyPositions[Keys.Indices[Keys.Indices.Num() -1]].OutputValue;
-				FKeyBlendInfo BlendStruct(KeysArray.Value.AllKeyPositions, Keys.Indices);
-				int32 CurrentIndex = 0;
-				for (int32 Index = 0; Index < NumIndices; ++Index)
-				{
-					const int32& KeyIndex = Keys.Indices[Index];
-					KeyHandles[Index] = KeysArray.Value.AllKeyHandles[KeyIndex];
-					const double CurrentValue = KeysArray.Value.AllKeyPositions[KeyIndex].OutputValue;
-					const double CurrentTime = KeysArray.Value.AllKeyPositions[KeyIndex].InputValue;
-					BlendStruct.SetValues(
-						PreviousTime, PreviousValue, CurrentTime, CurrentValue,
-						NextTime, NextValue, InBlendValue, FirstValue, LastValue, CurrentIndex
-						);
-					++CurrentIndex;
+			const int32 NumIndices = Indices.Num();
+			KeyHandles.SetNum(NumIndices);
+			KeyPositions.SetNum(NumIndices);
 
-					const double NewValue = InDoBlend(BlendStruct);
-					KeyPositions[Index] = FKeyPosition(KeysArray.Value.AllKeyPositions[KeyIndex].InputValue, NewValue);
-				}
-				Curve->SetKeyPositions(KeyHandles, KeyPositions);
-				bDidBlend = true;
-			}
+			// Move through every blended key and blend it
+			FKeyBlendInfo::ForEachKey(InBlendValue, BeforeBlendRange, AfterBlendRange, AllKeyPositions, Indices,
+				[&InDoBlend, &KeyPositions, &KeyHandles, &AllKeyHandles, &Indices](int32 Index, const FKeyBlendInfo& BlendInfo)
+				{
+					const double NewValue = InDoBlend(BlendInfo);
+					KeyPositions[Index] = FKeyPosition(BlendInfo.Current.X, NewValue);
+					KeyHandles[Index] = AllKeyHandles[Indices[Index]];
+				});
+				
+			Curve->SetKeyPositions(KeyHandles, KeyPositions);
+			bDidBlend = true;
 		}
 	}
 
diff --git a/Engine/Plugins/Animation/TweeningUtils/Source/TweeningUtilsEditor/Public/Math/KeyBlendInfo.h b/Engine/Plugins/Animation/TweeningUtils/Source/TweeningUtilsEditor/Public/Math/KeyBlendInfo.h
--- a/Engine/Plugins/Animation/TweeningUtils/Source/TweeningUtilsEditor/Public/Math/KeyBlendInfo.h
+++ b/Engine/Plugins/Animation/TweeningUtils/Source/TweeningUtilsEditor/Public/Math/KeyBlendInfo.h
@@ -2,55 +2,98 @@
 
 #pragma once
 
-#include "CurveDataAbstraction.h"
-#include "Containers/Array.h"
+#include "Containers/ArrayView.h"
 #include "HAL/Platform.h"
+#include "Math/Vector2D.h"
+#include <type_traits>
 
 namespace UE::TweeningUtilsEditor
 {
-/** Holds information about blending consecutive range of keys; each key has a time (x) and value (y). */
+/**
+ * Holds information about blending consecutive range of keys; each key has a time (x) and value (y).
+ * 
+ * This structure describes a range of consecutive keys to blend.
+ * Visually it looks like .xxxxxx., where '.'are keys that are not blended and 'x' keys that are blended.
+ * There can be no gap between the keys, so '.xx.xxx.' is not a blendable range: it would need to be split up into two ranges: '.xx." and '.xxx.'.
+ */
 struct FKeyBlendInfo
 {
-	// Previous and Next are the keys that are before/after the contiguous set of keys, 
-	// to get the Previous key for example you use AllKeyPostions[Indices[CurrentIndex -1]]
-	double PreviousTime; 
-	double PreviousValue; 
-	double CurrentTime; 
-	double CurrentValue;
-	double NextTime;
-	double NextValue;
+	// The values are ordered by key index.
+	
+	/** The key right before the range of blended keys. Same as FirstBlended if there is no such key. */
+	FVector2d BeforeBlendRange;
+
+	/** The first key in the range of blended keys. */
+	FVector2d FirstBlended;
+
+	/** The key before the currently blended one. If Current is the first key, this is BeforeBlendRange. */
+	FVector2d BeforeCurrent;
+
+	/** The key currently being blended. */
+	FVector2d Current;
+
+	/** The key after the current blended one. */
+	FVector2d AfterCurrent;
+
+	/** The last key in the range of blended keys. If current is the last key, this is AfterBlendRange. */
+	FVector2d LastBlended;
+
+	/** The key right after the range of blended keys. Same as LastBlended if there is no such key. */
+	FVector2d AfterBlendRange;
+	
 	/** The user specified value based on which to blend. Range is in [-1.0, 1.0]. */
 	double BlendValue;
-	double FirstValue;
-	double LastValue;
-
-	/** All keys, including the ones not being blended. */
-	TArray<FKeyPosition> AllKeyPositions;
-	/** Indices to AllKeyPositions which are being blended. */
-	TArray<int32> Indices;
-	/** Index to Indices. Identifies the key that is being blended right now. The blend algorithm moves through Indices by incrementing CurrentIndex. */
-	int32 CurrentIndex; 
-
-	explicit FKeyBlendInfo(const TArray<FKeyPosition>& InAllKeyPositions, const TArray<int32>& InIndices) 
-		: AllKeyPositions(InAllKeyPositions), Indices(InIndices)
+
+	explicit FKeyBlendInfo(
+		const FVector2d& InBeforeBlendRange, const FVector2d& InFirstBlended,
+		const FVector2d& InBeforeCurrent, const FVector2d& InCurrent, const FVector2d& InAfterCurrent,
+		const FVector2d& InLastBlended, const FVector2d& InAfterBlendRange,
+		double InBlendValue
+		)
+		: BeforeBlendRange(InBeforeBlendRange), FirstBlended(InFirstBlended)
+		, BeforeCurrent(InBeforeCurrent), Current(InCurrent), AfterCurrent(InAfterCurrent)
+		, LastBlended(InLastBlended), AfterBlendRange(InAfterBlendRange)
+		, BlendValue(InBlendValue)
 	{}
-	
-	void SetValues(
-		double InPreviousTime, double InPreviousValue, double InCurrentTime, double InCurrentValue,
-		double InNextTime, double InNextValue, double InBlendValue, double InFirstValue, double InLastValue,
-		int32 InCurrentIndex
+
+	/**
+	 * Given a consecutive range of keys to blend, this moves from left to right and invokes the callback with the corresponding FKeyBlendInfo.
+	 * @param InBlendValue The user specified value based on which to blend. Range is in [-1.0, 1.0]
+	 * @param InBeforeBlendRange The key right before the range of blended keys
+	 * @param InAfterBlendRange The key right after the range of blended keys
+	 * @param InAllKeyPositions Contains the keys to blend. It can contain more keys - only the keys specified in InBlendedIndices will be blended.
+	 * @param InBlendedIndices Indices of keys in InAllKeyPositions that are to be blended. It is assumed that the indices are sorted by time.
+	 * @param InCallback The callback that receives the current index of InBlendedIndices (i.e. InBlendedIndices.IsValidIndex(Index)) and
+	 *		the blend info for the currently blended key.
+	 */
+	template<typename TCallback> requires std::is_invocable_v<TCallback, int32 /*Index*/, const FKeyBlendInfo&>
+	static void ForEachKey(
+		double InBlendValue, const FVector2d& InBeforeBlendRange, const FVector2d& InAfterBlendRange,
+		const TConstArrayView<FVector2d>& InAllKeyPositions, const TConstArrayView<int32> InBlendedIndices,
+		TCallback&& InCallback
 		)
 	{
-		PreviousTime = InPreviousTime;
-		PreviousValue = InPreviousValue;
-		CurrentTime = InCurrentTime;
-		CurrentValue = InCurrentValue;
-		NextTime = InNextTime;
-		NextValue = InNextValue;
-		BlendValue = InBlendValue;
-		FirstValue = InFirstValue;
-		LastValue = InLastValue;
-		CurrentIndex = InCurrentIndex;
+		check(InBlendedIndices.Num() <= InAllKeyPositions.Num()); // Indices is supposed to be a consecutive sub-range of AllKeyPositions
+		const int32 NumIndices = InBlendedIndices.Num();
+				
+		const FVector2d& FirstBlended = InAllKeyPositions[InBlendedIndices[0]];
+		const FVector2d& LastBlended = InAllKeyPositions[InBlendedIndices[InBlendedIndices.Num() -1]];
+		for (int32 Index = 0; Index < NumIndices; ++Index)
+		{
+			const int32& KeyIndex = InBlendedIndices[Index];
+			const FVector2d& Current = InAllKeyPositions[KeyIndex];
+			const FVector2d& BeforeCurrent = Index > 0 ? InAllKeyPositions[InBlendedIndices[Index - 1]] : InBeforeBlendRange;
+			const FVector2d& AfterCurrent = Index < InBlendedIndices.Num() - 1 ? InAllKeyPositions[InBlendedIndices[Index + 1]] : InAfterBlendRange;
+					
+			const FKeyBlendInfo BlendStruct(
+				InBeforeBlendRange, FirstBlended,
+				BeforeCurrent, Current, AfterCurrent,
+				LastBlended, InAfterBlendRange,
+				InBlendValue
+				);
+
+			InCallback(Index, BlendStruct);
+		}
 	}
 };
 }
