TweeningUtils:
- Move key blending functions to runtime
- Move Widgets/MVC/Models to Math/Models

diff --git a/Engine/Plugins/Animation/ControlRig/Source/ControlRigEditor/Private/EditMode/Tween/ControlRigTweenModel.cpp b/Engine/Plugins/Animation/ControlRig/Source/ControlRigEditor/Private/EditMode/Tween/ControlRigTweenModel.cpp
--- a/Engine/Plugins/Animation/ControlRig/Source/ControlRigEditor/Private/EditMode/Tween/ControlRigTweenModel.cpp
+++ b/Engine/Plugins/Animation/ControlRig/Source/ControlRigEditor/Private/EditMode/Tween/ControlRigTweenModel.cpp
@@ -69,7 +69,7 @@ bool FControlRigTweenModel::DoBlend(float InNormalizedValue)
 	const float ScaledBlendValue = ScaleBlendValue(InNormalizedValue); 
 	if (!KeySelection.KeyMap.IsEmpty())
 	{
-		return TweeningUtilsEditor::BlendCurves(ScaledBlendValue, *CurveEditor, KeySelection, [this](const TweeningUtilsEditor::FKeyBlendInfo& InBlendInfo)
+		return TweeningUtilsEditor::BlendCurves(ScaledBlendValue, *CurveEditor, KeySelection, [this](const TweeningUtils::FKeyBlendInfo& InBlendInfo)
 		{
 			return GetTweenFunction().Blend(InBlendInfo);
 		});
@@ -103,7 +103,7 @@ bool FControlRigTweenModel::DoBlend(float InNormalizedValue)
 					const FVector2d Current { CurrentTime, CurrentValue };
 					
 					// First and last blended are the same as current since we're only blending a single key
-					const TweeningUtilsEditor::FKeyBlendInfo BlendStruct(
+					const TweeningUtils::FKeyBlendInfo BlendStruct(
 						BeforeBlendRange, Current, // Current is the last key to blend
 						// Key before current, current, and after current key
 						BeforeBlendRange, Current, AfterBlendRange, 
diff --git a/Engine/Plugins/Animation/ControlRig/Source/ControlRigEditor/Private/EditMode/Tween/ControlRigTweenModel.h b/Engine/Plugins/Animation/ControlRig/Source/ControlRigEditor/Private/EditMode/Tween/ControlRigTweenModel.h
--- a/Engine/Plugins/Animation/ControlRig/Source/ControlRigEditor/Private/EditMode/Tween/ControlRigTweenModel.h
+++ b/Engine/Plugins/Animation/ControlRig/Source/ControlRigEditor/Private/EditMode/Tween/ControlRigTweenModel.h
@@ -4,9 +4,9 @@
 
 #include "ControlRigObjectSelection.h"
 #include "Math/ContiguousKeyMapping.h"
+#include "Math/Models/FunctionBasedTweenModel.h"
 #include "Misc/Attribute.h"
 #include "Templates/SharedPointer.h"
-#include "Widgets/MVC/Models/FunctionBasedTweenModel.h"
 
 class ISequencer;
 
diff --git a/Engine/Plugins/Animation/ControlRig/Source/ControlRigEditor/Private/EditMode/Tween/TweenLogic.cpp b/Engine/Plugins/Animation/ControlRig/Source/ControlRigEditor/Private/EditMode/Tween/TweenLogic.cpp
--- a/Engine/Plugins/Animation/ControlRig/Source/ControlRigEditor/Private/EditMode/Tween/TweenLogic.cpp
+++ b/Engine/Plugins/Animation/ControlRig/Source/ControlRigEditor/Private/EditMode/Tween/TweenLogic.cpp
@@ -6,9 +6,9 @@
 #include "EditMode/ControlRigEditMode.h"
 #include "Framework/Commands/UICommandList.h"
 #include "Math/Abstraction/BuiltInTweenFunctionContainer.h"
+#include "Math/Models/EditorTweenModel.h"
 #include "Misc/SequencerUtils.h"
 #include "Templates/SharedPointer.h"
-#include "Widgets/MVC/Models/EditorTweenModel.h"
 
 namespace UE::ControlRigEditor
 {
diff --git a/Engine/Plugins/Animation/TweeningUtils/Source/TweeningUtils/Private/Math/KeyBlendingFunctions.cpp b/Engine/Plugins/Animation/TweeningUtils/Source/TweeningUtils/Private/Math/KeyBlendingFunctions.cpp
new file mode 100644
--- /dev/null
+++ b/Engine/Plugins/Animation/TweeningUtils/Source/TweeningUtils/Private/Math/KeyBlendingFunctions.cpp
@@ -0,0 +1,131 @@
+﻿// Copyright Epic Games, Inc. All Rights Reserved.
+
+#include "Math/KeyBlendingFunctions.h"
+
+#include "Math/KeyBlendInfo.h"
+#include "Math/Vector2D.h"
+
+namespace UE::TweeningUtils
+{
+double Blend_ControlsToTween(const FKeyBlendInfo& InBlendInfo)
+{
+	// Classic tween will move all to same location, not based on current time at all just blend and values
+	const double NormalizedBlendValue = (InBlendInfo.BlendValue + 1.0f) * 0.5f;
+	const double Value = InBlendInfo.BeforeBlendRange.Y + (InBlendInfo.AfterBlendRange.Y - InBlendInfo.BeforeBlendRange.Y) * (NormalizedBlendValue);
+	return Value;
+}
+
+double Blend_PushPull(const FKeyBlendInfo& InBlendInfo)
+{
+	if (FMath::IsNearlyEqual(InBlendInfo.AfterBlendRange.X, InBlendInfo.BeforeBlendRange.X))
+	{
+		return InBlendInfo.Current.Y;
+	}
+	const double T = (InBlendInfo.Current.X - InBlendInfo.BeforeBlendRange.X) / (InBlendInfo.AfterBlendRange.X - InBlendInfo.BeforeBlendRange.X);
+	const double ValueAtT = InBlendInfo.BeforeBlendRange.Y + T * (InBlendInfo.AfterBlendRange.Y - InBlendInfo.BeforeBlendRange.Y);
+	double NewValue;
+	if (InBlendInfo.BlendValue < 0.0)
+	{
+		NewValue = InBlendInfo.Current.Y + (-1.0 * InBlendInfo.BlendValue) * (ValueAtT - InBlendInfo.Current.Y);
+	}
+	else
+	{
+		const double AmplifyValueAtT = InBlendInfo.Current.Y + (InBlendInfo.Current.Y - ValueAtT);
+		NewValue = InBlendInfo.Current.Y + InBlendInfo.BlendValue * (AmplifyValueAtT - InBlendInfo.Current.Y);
+	}
+	return NewValue;
+}
+
+double Blend_Neighbor(const FKeyBlendInfo& InBlendInfo)
+{
+	double NewValue = InBlendInfo.Current.Y;
+	if (InBlendInfo.BlendValue < 0.0)
+	{
+		NewValue = InBlendInfo.Current.Y + (-1.0 * InBlendInfo.BlendValue) * (InBlendInfo.BeforeBlendRange.Y - InBlendInfo.Current.Y);
+	}
+	else
+	{
+		NewValue = InBlendInfo.Current.Y + InBlendInfo.BlendValue * (InBlendInfo.AfterBlendRange.Y - InBlendInfo.Current.Y);
+	}
+	return NewValue;
+}
+
+double Blend_Relative(const FKeyBlendInfo& InBlendInfo)
+{
+	double NewValue = InBlendInfo.Current.Y;
+	if (InBlendInfo.BlendValue < 0.0)
+	{
+		NewValue = InBlendInfo.Current.Y + (-1.0 * InBlendInfo.BlendValue) * (InBlendInfo.BeforeBlendRange.Y - InBlendInfo.FirstBlended.Y);
+	}
+	else
+	{
+		NewValue = InBlendInfo.Current.Y + InBlendInfo.BlendValue * (InBlendInfo.AfterBlendRange.Y - InBlendInfo.LastBlended.Y);
+	}
+	return NewValue;
+}
+
+namespace Private::BlendToEase
+{ 
+static float SCurve(float X, float Slope, float Width, float Height, float XShift, float YShift)
+{
+	if (X > (XShift + Width))
+	{
+		return(Height + YShift);
+	}
+	if (X < XShift)
+	{
+		return YShift;
+	}
+
+	const float Val = Height * (FMath::Pow((X - XShift), Slope) / (FMath::Pow(X - XShift, Slope) + FMath::Pow((Width - (X - XShift)), Slope))) + YShift;
+	return Val;
+}
+}
+
+double Blend_Ease(const FKeyBlendInfo& InBlendInfo)
+{
+	if (FMath::IsNearlyEqual(InBlendInfo.AfterBlendRange.X, InBlendInfo.BeforeBlendRange.X))
+	{
+		return InBlendInfo.Current.Y;
+	}
+	const double Source = InBlendInfo.Current.Y;
+	const double FullTimeDiff = InBlendInfo.AfterBlendRange.X - InBlendInfo.BeforeBlendRange.X;
+	const double AbsValue = FMath::Abs(InBlendInfo.BlendValue);
+	const double X = InBlendInfo.Current.X - InBlendInfo.BeforeBlendRange.X;
+	const double Ratio = X / FullTimeDiff;
+	double Shift, Delta, Base;
+	if (InBlendInfo.BlendValue > 0)
+	{
+		Shift = -1.0;
+		Delta = InBlendInfo.AfterBlendRange.Y - Source;
+		Base = Source;
+	}
+	else
+	{
+		Shift = 0.0;
+		Delta = Source - InBlendInfo.BeforeBlendRange.Y;
+		Base = InBlendInfo.BeforeBlendRange.Y;
+	}
+	const double Slope = 5.0 * AbsValue;
+	const double EaseY = Private::BlendToEase::SCurve(Ratio, Slope, 2.0, 2.0, Shift, Shift);
+	const double NewValue = Base + (Delta * EaseY);
+
+	return NewValue;
+}
+
+double Blend_SmoothRough(const FKeyBlendInfo& InBlendInfo)
+{
+	const double PrevVal = InBlendInfo.BeforeCurrent.Y;
+	const double CurVal = InBlendInfo.Current.Y;
+	const double NextVal = InBlendInfo.AfterCurrent.Y;
+	const double NewValue = (PrevVal * 0.25 + CurVal * 0.5 + NextVal * 0.25);
+	if (InBlendInfo.BlendValue < 0.0)
+	{
+		return InBlendInfo.Current.Y + (-1.0 * InBlendInfo.BlendValue) * (NewValue - InBlendInfo.Current.Y);
+	}
+	else
+	{
+		return InBlendInfo.Current.Y + InBlendInfo.BlendValue * (InBlendInfo.Current.Y - NewValue);
+	}
+}
+}
diff --git a/Engine/Plugins/Animation/TweeningUtils/Source/TweeningUtils/Private/TweeningUtilsModule.cpp b/Engine/Plugins/Animation/TweeningUtils/Source/TweeningUtils/Private/TweeningUtilsModule.cpp
new file mode 100644
--- /dev/null
+++ b/Engine/Plugins/Animation/TweeningUtils/Source/TweeningUtils/Private/TweeningUtilsModule.cpp
@@ -0,0 +1,15 @@
+﻿// Copyright Epic Games, Inc. All Rights Reserved.
+
+#include "TweeningUtilsModule.h"
+
+namespace UE::TweeningUtils
+{
+void FTweeningUtilsModule::StartupModule()
+{}
+
+void FTweeningUtilsModule::ShutdownModule()
+{}
+	
+IMPLEMENT_MODULE(FTweeningUtilsModule, TweeningUtils)
+}
+    
\ No newline at end of file
diff --git a/Engine/Plugins/Animation/TweeningUtils/Source/TweeningUtils/Private/TweeningUtilsModule.h b/Engine/Plugins/Animation/TweeningUtils/Source/TweeningUtils/Private/TweeningUtilsModule.h
new file mode 100644
--- /dev/null
+++ b/Engine/Plugins/Animation/TweeningUtils/Source/TweeningUtils/Private/TweeningUtilsModule.h
@@ -0,0 +1,15 @@
+﻿// Copyright Epic Games, Inc. All Rights Reserved.
+
+#pragma once
+
+#include "Modules/ModuleManager.h"
+
+namespace UE::TweeningUtils
+{
+class FTweeningUtilsModule : public IModuleInterface
+{
+public:
+	virtual void StartupModule() override;
+	virtual void ShutdownModule() override;
+};
+}
diff --git a/Engine/Plugins/Animation/TweeningUtils/Source/TweeningUtils/Public/Math/KeyBlendInfo.h b/Engine/Plugins/Animation/TweeningUtils/Source/TweeningUtils/Public/Math/KeyBlendInfo.h
new file mode 100644
--- /dev/null
+++ b/Engine/Plugins/Animation/TweeningUtils/Source/TweeningUtils/Public/Math/KeyBlendInfo.h
@@ -0,0 +1,99 @@
+﻿// Copyright Epic Games, Inc. All Rights Reserved.
+
+#pragma once
+
+#include "Containers/ArrayView.h"
+#include "HAL/Platform.h"
+#include "Math/Vector2D.h"
+#include <type_traits>
+
+namespace UE::TweeningUtils
+{
+/**
+ * Holds information about blending consecutive range of keys; each key has a time (x) and value (y).
+ * 
+ * This structure describes a range of consecutive keys to blend.
+ * Visually it looks like .xxxxxx., where '.'are keys that are not blended and 'x' keys that are blended.
+ * There can be no gap between the keys, so '.xx.xxx.' is not a blendable range: it would need to be split up into two ranges: '.xx." and '.xxx.'.
+ */
+struct FKeyBlendInfo
+{
+	// The values are ordered by key index.
+	
+	/** The key right before the range of blended keys. Same as FirstBlended if there is no such key. */
+	FVector2d BeforeBlendRange;
+
+	/** The first key in the range of blended keys. */
+	FVector2d FirstBlended;
+
+	/** The key before the currently blended one. If Current is the first key, this is BeforeBlendRange. */
+	FVector2d BeforeCurrent;
+
+	/** The key currently being blended. */
+	FVector2d Current;
+
+	/** The key after the current blended one. */
+	FVector2d AfterCurrent;
+
+	/** The last key in the range of blended keys. If current is the last key, this is AfterBlendRange. */
+	FVector2d LastBlended;
+
+	/** The key right after the range of blended keys. Same as LastBlended if there is no such key. */
+	FVector2d AfterBlendRange;
+	
+	/** The user specified value based on which to blend. Range is in [-1.0, 1.0]. */
+	double BlendValue;
+
+	explicit FKeyBlendInfo(
+		const FVector2d& InBeforeBlendRange, const FVector2d& InFirstBlended,
+		const FVector2d& InBeforeCurrent, const FVector2d& InCurrent, const FVector2d& InAfterCurrent,
+		const FVector2d& InLastBlended, const FVector2d& InAfterBlendRange,
+		double InBlendValue
+		)
+		: BeforeBlendRange(InBeforeBlendRange), FirstBlended(InFirstBlended)
+		, BeforeCurrent(InBeforeCurrent), Current(InCurrent), AfterCurrent(InAfterCurrent)
+		, LastBlended(InLastBlended), AfterBlendRange(InAfterBlendRange)
+		, BlendValue(InBlendValue)
+	{}
+
+	/**
+	 * Given a consecutive range of keys to blend, this moves from left to right and invokes the callback with the corresponding FKeyBlendInfo.
+	 * @param InBlendValue The user specified value based on which to blend. Range is in [-1.0, 1.0]
+	 * @param InBeforeBlendRange The key right before the range of blended keys
+	 * @param InAfterBlendRange The key right after the range of blended keys
+	 * @param InAllKeyPositions Contains the keys to blend. It can contain more keys - only the keys specified in InBlendedIndices will be blended.
+	 * @param InBlendedIndices Indices of keys in InAllKeyPositions that are to be blended. It is assumed that the indices are sorted by time.
+	 * @param InCallback The callback that receives the current index of InBlendedIndices (i.e. InBlendedIndices.IsValidIndex(Index)) and
+	 *		the blend info for the currently blended key.
+	 */
+	template<typename TCallback> requires std::is_invocable_v<TCallback, int32 /*Index*/, const FKeyBlendInfo&>
+	static void ForEachKey(
+		double InBlendValue, const FVector2d& InBeforeBlendRange, const FVector2d& InAfterBlendRange,
+		const TConstArrayView<FVector2d>& InAllKeyPositions, const TConstArrayView<int32> InBlendedIndices,
+		TCallback&& InCallback
+		)
+	{
+		check(InBlendedIndices.Num() <= InAllKeyPositions.Num()); // Indices is supposed to be a consecutive sub-range of AllKeyPositions
+		const int32 NumIndices = InBlendedIndices.Num();
+				
+		const FVector2d& FirstBlended = InAllKeyPositions[InBlendedIndices[0]];
+		const FVector2d& LastBlended = InAllKeyPositions[InBlendedIndices[InBlendedIndices.Num() -1]];
+		for (int32 Index = 0; Index < NumIndices; ++Index)
+		{
+			const int32& KeyIndex = InBlendedIndices[Index];
+			const FVector2d& Current = InAllKeyPositions[KeyIndex];
+			const FVector2d& BeforeCurrent = Index > 0 ? InAllKeyPositions[InBlendedIndices[Index - 1]] : InBeforeBlendRange;
+			const FVector2d& AfterCurrent = Index < InBlendedIndices.Num() - 1 ? InAllKeyPositions[InBlendedIndices[Index + 1]] : InAfterBlendRange;
+					
+			const FKeyBlendInfo BlendStruct(
+				InBeforeBlendRange, FirstBlended,
+				BeforeCurrent, Current, AfterCurrent,
+				LastBlended, InAfterBlendRange,
+				InBlendValue
+				);
+
+			InCallback(Index, BlendStruct);
+		}
+	}
+};
+}
diff --git a/Engine/Plugins/Animation/TweeningUtils/Source/TweeningUtils/Public/Math/KeyBlendingFunctions.h b/Engine/Plugins/Animation/TweeningUtils/Source/TweeningUtils/Public/Math/KeyBlendingFunctions.h
new file mode 100644
--- /dev/null
+++ b/Engine/Plugins/Animation/TweeningUtils/Source/TweeningUtils/Public/Math/KeyBlendingFunctions.h
@@ -0,0 +1,20 @@
+﻿// Copyright Epic Games, Inc. All Rights Reserved.
+
+#pragma once
+
+#include "KeyBlendInfo.h"
+
+namespace UE::TweeningUtils
+{
+TWEENINGUTILS_API double Blend_ControlsToTween(const FKeyBlendInfo& InBlendInfo);
+
+TWEENINGUTILS_API double Blend_PushPull(const FKeyBlendInfo& InBlendInfo);
+
+TWEENINGUTILS_API double Blend_Neighbor(const FKeyBlendInfo& InBlendInfo);
+
+TWEENINGUTILS_API double Blend_Relative(const FKeyBlendInfo& InBlendInfo);
+
+TWEENINGUTILS_API double Blend_Ease(const FKeyBlendInfo& InBlendInfo);
+
+TWEENINGUTILS_API double Blend_SmoothRough(const FKeyBlendInfo& InBlendInfo);
+}
diff --git a/Engine/Plugins/Animation/TweeningUtils/Source/TweeningUtils/TweeningUtils.Build.cs b/Engine/Plugins/Animation/TweeningUtils/Source/TweeningUtils/TweeningUtils.Build.cs
new file mode 100644
--- /dev/null
+++ b/Engine/Plugins/Animation/TweeningUtils/Source/TweeningUtils/TweeningUtils.Build.cs
@@ -0,0 +1,23 @@
+﻿using UnrealBuildTool;
+
+public class TweeningUtils : ModuleRules
+{
+    public TweeningUtils(ReadOnlyTargetRules Target) : base(Target)
+    {
+        PCHUsage = ModuleRules.PCHUsageMode.UseExplicitOrSharedPCHs;
+
+        PublicDependencyModuleNames.AddRange(
+            new string[]
+            {
+                "Core",
+            }
+        );
+
+        PrivateDependencyModuleNames.AddRange(
+            new string[]
+            {
+                "Engine"
+            }
+        );
+    }
+}
\ No newline at end of file
diff --git a/Engine/Plugins/Animation/TweeningUtils/Source/TweeningUtilsEditor/Private/Math/Abstraction/KeyBlendingAbstraction.cpp b/Engine/Plugins/Animation/TweeningUtils/Source/TweeningUtilsEditor/Private/Math/Abstraction/KeyBlendingAbstraction.cpp
--- a/Engine/Plugins/Animation/TweeningUtils/Source/TweeningUtilsEditor/Private/Math/Abstraction/KeyBlendingAbstraction.cpp
+++ b/Engine/Plugins/Animation/TweeningUtils/Source/TweeningUtilsEditor/Private/Math/Abstraction/KeyBlendingAbstraction.cpp
@@ -20,7 +20,7 @@ namespace UE::TweeningUtilsEditor
 {
 struct FStaticBlendFunctionData
 {
-	using FBlendFunctionPtr = double(*)(const FKeyBlendInfo& InBlendInfo);
+	using FBlendFunctionPtr = double(*)(const TweeningUtils::FKeyBlendInfo& InBlendInfo);
 	
 	const EBlendFunction Type;
 	const FBlendFunctionPtr EvaluationFunctionPtr;
@@ -72,42 +72,42 @@ static FStaticBlendFunctionData& GetFunctionData(EBlendFunction InFunction)
 	static FStaticBlendFunctionData Data[Num] =
 	{
 		FStaticBlendFunctionData(
-			EBlendFunction::BlendNeighbor, &Blend_Neighbor, FTweeningUtilsCommands::Get().SetTweenBlendNeighbor,
+			EBlendFunction::BlendNeighbor, &TweeningUtils::Blend_Neighbor, FTweeningUtilsCommands::Get().SetTweenBlendNeighbor,
 			LOCTEXT("BlendNeighbor.Label", "Blend Neighbor"),
 			LOCTEXT("BlendNeighbour.Abbreviation", "BN"),
 			LOCTEXT("BlendNeighbour.Description", "Blend to the next or previous values for selected keys."),
 			TEXT("BlendNeighbour")
 		),
 		FStaticBlendFunctionData(
-			EBlendFunction::PushPull, &Blend_PushPull, FTweeningUtilsCommands::Get().SetTweenPushPull,
+			EBlendFunction::PushPull, &TweeningUtils::Blend_PushPull, FTweeningUtilsCommands::Get().SetTweenPushPull,
 			LOCTEXT("PushPull.Label", "Push / Pull"),
 			LOCTEXT("PushPull.Abbreviation", "PP"),
 			LOCTEXT("PushPull.Description", "Push or pull the values to the interpolation between the previous and next keys"),
 			TEXT("PushPull")
 		),
 		FStaticBlendFunctionData(
-			EBlendFunction::BlendEase, &Blend_Ease, FTweeningUtilsCommands::Get().SetTweenBlendEase,
+			EBlendFunction::BlendEase, &TweeningUtils::Blend_Ease, FTweeningUtilsCommands::Get().SetTweenBlendEase,
 			LOCTEXT("BlendEase.Label", "Blend Ease"),
 			LOCTEXT("BlendEase.Abbreviation", "BE"),
 			LOCTEXT("BlendEase.Description", "Blend with an ease falloff to the next or previous value for selected keys"),
 			TEXT("BlendEase")
 		),
 		FStaticBlendFunctionData(
-			EBlendFunction::ControlsToTween, &Blend_ControlsToTween, FTweeningUtilsCommands::Get().SetControlsToTween,
+			EBlendFunction::ControlsToTween, &TweeningUtils::Blend_ControlsToTween, FTweeningUtilsCommands::Get().SetControlsToTween,
 			LOCTEXT("ControlsToTween.Label", "Tween"),
 			LOCTEXT("ControlsToTween.Abbreviation", "TW"),
 			LOCTEXT("ControlsToTween.Description", "Interpolates between the previous and next keys"),
 			TEXT("Tween")
 		),
 		FStaticBlendFunctionData(
-			EBlendFunction::BlendRelative, &Blend_Relative, FTweeningUtilsCommands::Get().SetTweenBlendRelative,
+			EBlendFunction::BlendRelative, &TweeningUtils::Blend_Relative, FTweeningUtilsCommands::Get().SetTweenBlendRelative,
 			LOCTEXT("BlendRelative.Label", "Blend Relative"),
 			LOCTEXT("BlendRelative.Abbreviation", "BR"),
 			LOCTEXT("BlendRelative.Description", "Blend relative to the next or previous value for selected keys"),
 			TEXT("BlendRelative")
 		),
 		FStaticBlendFunctionData(
-			EBlendFunction::SmoothRough, &Blend_SmoothRough, FTweeningUtilsCommands::Get().SetTweenSmoothRough,
+			EBlendFunction::SmoothRough, &TweeningUtils::Blend_SmoothRough, FTweeningUtilsCommands::Get().SetTweenSmoothRough,
 			LOCTEXT("SmoothRough.Label", "Smooth Rough"),
 			LOCTEXT("SmoothRough.Abbreviation", "SR"),
 			LOCTEXT("SmoothRough.Description", "Make selected keys smooth or rough"),
@@ -137,7 +137,7 @@ TOptional<EBlendFunction> LexBlendFunction(const FString& InString)
 	return {};
 }
 	
-double EvaluateBlend(EBlendFunction InFunction, const FKeyBlendInfo& InBlendInfo)
+double EvaluateBlend(EBlendFunction InFunction, const TweeningUtils::FKeyBlendInfo& InBlendInfo)
 {
 	return GetFunctionData(InFunction).EvaluationFunctionPtr(InBlendInfo);
 }
diff --git a/Engine/Plugins/Animation/TweeningUtils/Source/TweeningUtilsEditor/Private/Math/KeyBlendingFunctions.cpp b/Engine/Plugins/Animation/TweeningUtils/Source/TweeningUtilsEditor/Private/Math/KeyBlendingFunctions.cpp
deleted file mode 100644
diff --git a/Engine/Plugins/Animation/TweeningUtils/Source/TweeningUtilsEditor/Private/Math/Models/CurveTweeningModel.cpp b/Engine/Plugins/Animation/TweeningUtils/Source/TweeningUtilsEditor/Private/Math/Models/CurveTweeningModel.cpp
new file mode 100644
--- /dev/null
+++ b/Engine/Plugins/Animation/TweeningUtils/Source/TweeningUtilsEditor/Private/Math/Models/CurveTweeningModel.cpp
@@ -0,0 +1,48 @@
+﻿// Copyright Epic Games, Inc. All Rights Reserved.
+
+#include "Math/Models/CurveTweenModel.h"
+
+#include "Math/CurveBlending.h"
+#include "Math/Abstraction/ITweenFunction.h"
+
+namespace UE::TweeningUtilsEditor
+{
+FCurveTweenModel::FCurveTweenModel(TWeakPtr<FCurveEditor> InWeakCurveEditor, const ITweenFunction& InTweenFunction)
+	: FFunctionBasedTweenModel(InTweenFunction)
+	, WeakCurveEditor(InWeakCurveEditor)
+{}
+
+bool FCurveTweenModel::HasAnythingToBlend() const
+{
+	const TSharedPtr<FCurveEditor> CurveEditor = WeakCurveEditor.Pin();
+	return CurveEditor && !CurveEditor->GetSelection().IsEmpty();
+}
+
+void FCurveTweenModel::StartBlendOperation()
+{
+	const TSharedPtr<FCurveEditor> CurveEditor = WeakCurveEditor.Pin();
+	ContiguousKeySelection = CurveEditor ? FContiguousKeyMapping(*CurveEditor) : FContiguousKeyMapping{};
+}
+
+void FCurveTweenModel::StopBlendOperation()
+{
+	ContiguousKeySelection.KeyMap.Empty(
+		ContiguousKeySelection.KeyMap.Num()
+		);
+}
+
+void FCurveTweenModel::BlendValues(float InNormalizedValue)
+{
+	const TSharedPtr<FCurveEditor> CurveEditor = WeakCurveEditor.Pin();
+	if (!CurveEditor)
+	{
+		return;
+	}
+
+	const float ScaledBlendValue = ScaleBlendValue(InNormalizedValue); 
+	TweeningUtilsEditor::BlendCurves(ScaledBlendValue, *CurveEditor, ContiguousKeySelection, [this](const TweeningUtils::FKeyBlendInfo& InBlendInfo)
+	{
+		return GetTweenFunction().Blend(InBlendInfo);
+	});
+}
+}
diff --git a/Engine/Plugins/Animation/TweeningUtils/Source/TweeningUtilsEditor/Private/Widgets/MVC/CycleFunctionController.cpp b/Engine/Plugins/Animation/TweeningUtils/Source/TweeningUtilsEditor/Private/Widgets/MVC/CycleFunctionController.cpp
--- a/Engine/Plugins/Animation/TweeningUtils/Source/TweeningUtilsEditor/Private/Widgets/MVC/CycleFunctionController.cpp
+++ b/Engine/Plugins/Animation/TweeningUtils/Source/TweeningUtilsEditor/Private/Widgets/MVC/CycleFunctionController.cpp
@@ -5,8 +5,8 @@
 #include "Framework/Commands/UICommandInfo.h"
 #include "Framework/Commands/UICommandList.h"
 #include "Math/Abstraction/ITweenFunctionContainer.h"
+#include "Math/Models/FunctionBasedTweenModel.h"
 #include "TweeningUtilsCommands.h"
-#include "Widgets/MVC/Models/FunctionBasedTweenModel.h"
 
 namespace UE::TweeningUtilsEditor
 {
diff --git a/Engine/Plugins/Animation/TweeningUtils/Source/TweeningUtilsEditor/Private/Widgets/MVC/Models/CurveTweeningModel.cpp b/Engine/Plugins/Animation/TweeningUtils/Source/TweeningUtilsEditor/Private/Widgets/MVC/Models/CurveTweeningModel.cpp
deleted file mode 100644
diff --git a/Engine/Plugins/Animation/TweeningUtils/Source/TweeningUtilsEditor/Private/Widgets/MVC/STweenView.cpp b/Engine/Plugins/Animation/TweeningUtils/Source/TweeningUtilsEditor/Private/Widgets/MVC/STweenView.cpp
--- a/Engine/Plugins/Animation/TweeningUtils/Source/TweeningUtilsEditor/Private/Widgets/MVC/STweenView.cpp
+++ b/Engine/Plugins/Animation/TweeningUtils/Source/TweeningUtilsEditor/Private/Widgets/MVC/STweenView.cpp
@@ -2,9 +2,9 @@
 
 #include "Widgets/MVC/STweenView.h"
 
+#include "Math/Models/TweenModel.h"
 #include "Styling/AppStyle.h"
 #include "Widgets/Layout/SBox.h"
-#include "Widgets/MVC/Models/TweenModel.h"
 #include "Widgets/STweenSlider.h"
 
 namespace UE::TweeningUtilsEditor
diff --git a/Engine/Plugins/Animation/TweeningUtils/Source/TweeningUtilsEditor/Private/Widgets/MVC/TweenControllers.cpp b/Engine/Plugins/Animation/TweeningUtils/Source/TweeningUtilsEditor/Private/Widgets/MVC/TweenControllers.cpp
--- a/Engine/Plugins/Animation/TweeningUtils/Source/TweeningUtilsEditor/Private/Widgets/MVC/TweenControllers.cpp
+++ b/Engine/Plugins/Animation/TweeningUtils/Source/TweeningUtilsEditor/Private/Widgets/MVC/TweenControllers.cpp
@@ -2,7 +2,7 @@
 
 #include "Widgets/MVC/TweenControllers.h"
 
-#include "Widgets/MVC/Models/FunctionBasedTweenModel.h"
+#include "Math/Models/FunctionBasedTweenModel.h"
 
 namespace UE::TweeningUtilsEditor
 {
diff --git a/Engine/Plugins/Animation/TweeningUtils/Source/TweeningUtilsEditor/Private/Widgets/MVC/TweenMouseSlidingController.cpp b/Engine/Plugins/Animation/TweeningUtils/Source/TweeningUtilsEditor/Private/Widgets/MVC/TweenMouseSlidingController.cpp
--- a/Engine/Plugins/Animation/TweeningUtils/Source/TweeningUtilsEditor/Private/Widgets/MVC/TweenMouseSlidingController.cpp
+++ b/Engine/Plugins/Animation/TweeningUtils/Source/TweeningUtilsEditor/Private/Widgets/MVC/TweenMouseSlidingController.cpp
@@ -3,11 +3,11 @@
 #include "Widgets/MVC/TweenMouseSlidingController.h"
 
 #include "HAL/IConsoleManager.h"
+#include "Math/Models/TweenModel.h"
 #include "Misc/Attribute.h"
 #include "TweeningUtilsCommands.h"
 #include "TweeningUtilsStyle.h"
 #include "Widgets/TweenSliderStyle.h"
-#include "Widgets/MVC/Models/TweenModel.h"
 
 namespace UE::TweeningUtilsEditor
 {
diff --git a/Engine/Plugins/Animation/TweeningUtils/Source/TweeningUtilsEditor/Private/Widgets/MVC/TweenToolbarController.cpp b/Engine/Plugins/Animation/TweeningUtils/Source/TweeningUtilsEditor/Private/Widgets/MVC/TweenToolbarController.cpp
--- a/Engine/Plugins/Animation/TweeningUtils/Source/TweeningUtilsEditor/Private/Widgets/MVC/TweenToolbarController.cpp
+++ b/Engine/Plugins/Animation/TweeningUtils/Source/TweeningUtilsEditor/Private/Widgets/MVC/TweenToolbarController.cpp
@@ -5,12 +5,12 @@
 #include "Framework/Commands/UICommandList.h"
 #include "Framework/MultiBox/MultiBoxBuilder.h"
 #include "Math/Abstraction/ITweenFunctionContainer.h"
+#include "Math/Models/FunctionBasedTweenModel.h"
 #include "Math/Color.h"
 #include "Styling/SlateStyle.h"
 #include "Styling/SlateStyleMacros.h"
 #include "Styling/SlateStyleRegistry.h"
 #include "TweeningUtilsCommands.h"
-#include "Widgets/MVC/Models/FunctionBasedTweenModel.h"
 #include "Widgets/MVC/STweenView.h"
 
 namespace UE::TweeningUtilsEditor
diff --git a/Engine/Plugins/Animation/TweeningUtils/Source/TweeningUtilsEditor/Public/Math/Abstraction/BuiltInTweenFunction.h b/Engine/Plugins/Animation/TweeningUtils/Source/TweeningUtilsEditor/Public/Math/Abstraction/BuiltInTweenFunction.h
--- a/Engine/Plugins/Animation/TweeningUtils/Source/TweeningUtilsEditor/Public/Math/Abstraction/BuiltInTweenFunction.h
+++ b/Engine/Plugins/Animation/TweeningUtils/Source/TweeningUtilsEditor/Public/Math/Abstraction/BuiltInTweenFunction.h
@@ -7,7 +7,6 @@
 
 namespace UE::TweeningUtilsEditor
 {
-
 /** Evaluates a blend function that is identified by EBlendFunction. */
 class FBuiltInTweenFunction : public ITweenFunction
 {
@@ -19,7 +18,7 @@ public:
 	void SetBlendFunction(EBlendFunction InValue) { BlendFunction = InValue; }
 	
 	//~ Begin ITweenFunction Interface
-	virtual double Blend(const FKeyBlendInfo& InBlendInfo) const override { return EvaluateBlend(BlendFunction, InBlendInfo); }
+	virtual double Blend(const TweeningUtils::FKeyBlendInfo& InBlendInfo) const override { return EvaluateBlend(BlendFunction, InBlendInfo); }
 	//~ End ITweenFunction Interface
 
 private:
diff --git a/Engine/Plugins/Animation/TweeningUtils/Source/TweeningUtilsEditor/Public/Math/Abstraction/ITweenFunction.h b/Engine/Plugins/Animation/TweeningUtils/Source/TweeningUtilsEditor/Public/Math/Abstraction/ITweenFunction.h
--- a/Engine/Plugins/Animation/TweeningUtils/Source/TweeningUtilsEditor/Public/Math/Abstraction/ITweenFunction.h
+++ b/Engine/Plugins/Animation/TweeningUtils/Source/TweeningUtilsEditor/Public/Math/Abstraction/ITweenFunction.h
@@ -2,9 +2,10 @@
 
 #pragma once
 
+namespace UE::TweeningUtils { struct FKeyBlendInfo; }
+
 namespace UE::TweeningUtilsEditor
 {
-struct FKeyBlendInfo;
 
 /** Abstracts tweening functions. This is useful for maintaining generic UI but support adding custom functions from external modules. */
 class ITweenFunction
@@ -12,7 +13,7 @@ class ITweenFunction
 public:
 
 	/** @return The result of combining FKeyBlendInfo::CurrentValue and FKeyBlendInfo::BlendValue according to this function. */
-	virtual double Blend(const FKeyBlendInfo& InBlendInfo) const = 0;
+	virtual double Blend(const TweeningUtils::FKeyBlendInfo& InBlendInfo) const = 0;
 
 	virtual ~ITweenFunction() = default;
 };
diff --git a/Engine/Plugins/Animation/TweeningUtils/Source/TweeningUtilsEditor/Public/Math/Abstraction/KeyBlendingAbstraction.h b/Engine/Plugins/Animation/TweeningUtils/Source/TweeningUtilsEditor/Public/Math/Abstraction/KeyBlendingAbstraction.h
--- a/Engine/Plugins/Animation/TweeningUtils/Source/TweeningUtilsEditor/Public/Math/Abstraction/KeyBlendingAbstraction.h
+++ b/Engine/Plugins/Animation/TweeningUtils/Source/TweeningUtilsEditor/Public/Math/Abstraction/KeyBlendingAbstraction.h
@@ -13,12 +13,11 @@ class FUICommandInfo;
 struct FLinearColor;
 struct FSlateBrush;
 template<typename T> struct TOptional;
+namespace UE::TweeningUtils { struct FKeyBlendInfo; }
 
 /** The following just contains some utils that is common in UI code. Tools can opt for their own description but using this helps with consistency. */
 namespace UE::TweeningUtilsEditor
 {
-struct FKeyBlendInfo;
-	
 /** Util enum which is useful for UI code to abstract functions.*/
 enum class EBlendFunction : uint8
 {
@@ -50,7 +49,7 @@ TWEENINGUTILSEDITOR_API FString BlendFunctionToString(EBlendFunction InFunction)
 TWEENINGUTILSEDITOR_API TOptional<EBlendFunction> LexBlendFunction(const FString& InString);
 	
 /** Evaluates the function */
-TWEENINGUTILSEDITOR_API double EvaluateBlend(EBlendFunction InFunction, const FKeyBlendInfo& InBlendInfo);
+TWEENINGUTILSEDITOR_API double EvaluateBlend(EBlendFunction InFunction, const TweeningUtils::FKeyBlendInfo& InBlendInfo);
 
 /** @return The command that is used to select the function in most UI. */
 TWEENINGUTILSEDITOR_API TSharedPtr<FUICommandInfo> GetCommandForBlendFunction(EBlendFunction InFunction);
diff --git a/Engine/Plugins/Animation/TweeningUtils/Source/TweeningUtilsEditor/Public/Math/CurveBlending.h b/Engine/Plugins/Animation/TweeningUtils/Source/TweeningUtilsEditor/Public/Math/CurveBlending.h
--- a/Engine/Plugins/Animation/TweeningUtils/Source/TweeningUtilsEditor/Public/Math/CurveBlending.h
+++ b/Engine/Plugins/Animation/TweeningUtils/Source/TweeningUtilsEditor/Public/Math/CurveBlending.h
@@ -23,13 +23,13 @@ namespace UE::TweeningUtilsEditor
  * @param InDoBlend	Does the blending
  * @return Whether InDoBlend was invokes at least once.
  */
-template<typename TCallback> requires std::is_invocable_r_v<double, TCallback, const FKeyBlendInfo&>
+template<typename TCallback> requires std::is_invocable_r_v<double, TCallback, const TweeningUtils::FKeyBlendInfo&>
 bool BlendCurves(float InBlendValue, const FCurveEditor& InCurveEditor, const FContiguousKeyMapping& InKeySelection, TCallback&& InDoBlend);
 }
 
 namespace UE::TweeningUtilsEditor
 {
-template<typename TCallback> requires std::is_invocable_r_v<double, TCallback, const FKeyBlendInfo&>
+template<typename TCallback> requires std::is_invocable_r_v<double, TCallback, const TweeningUtils::FKeyBlendInfo&>
 bool BlendCurves(float InBlendValue, const FCurveEditor& InCurveEditor, const FContiguousKeyMapping& InKeySelection, TCallback&& InDoBlend)
 {
 	bool bDidBlend = false;
@@ -62,8 +62,8 @@ bool BlendCurves(float InBlendValue, const FCurveEditor& InCurveEditor, const FC
 			KeyPositions.SetNum(NumIndices);
 
 			// Move through every blended key and blend it
-			FKeyBlendInfo::ForEachKey(InBlendValue, BeforeBlendRange, AfterBlendRange, AllKeyPositions, Indices,
-				[&InDoBlend, &KeyPositions, &KeyHandles, &AllKeyHandles, &Indices](int32 Index, const FKeyBlendInfo& BlendInfo)
+			TweeningUtils::FKeyBlendInfo::ForEachKey(InBlendValue, BeforeBlendRange, AfterBlendRange, AllKeyPositions, Indices,
+				[&InDoBlend, &KeyPositions, &KeyHandles, &AllKeyHandles, &Indices](int32 Index, const TweeningUtils::FKeyBlendInfo& BlendInfo)
 				{
 					const double NewValue = InDoBlend(BlendInfo);
 					KeyPositions[Index] = FKeyPosition(BlendInfo.Current.X, NewValue);
diff --git a/Engine/Plugins/Animation/TweeningUtils/Source/TweeningUtilsEditor/Public/Math/KeyBlendInfo.h b/Engine/Plugins/Animation/TweeningUtils/Source/TweeningUtilsEditor/Public/Math/KeyBlendInfo.h
deleted file mode 100644
diff --git a/Engine/Plugins/Animation/TweeningUtils/Source/TweeningUtilsEditor/Public/Math/KeyBlendingFunctions.h b/Engine/Plugins/Animation/TweeningUtils/Source/TweeningUtilsEditor/Public/Math/KeyBlendingFunctions.h
deleted file mode 100644
diff --git a/Engine/Plugins/Animation/TweeningUtils/Source/TweeningUtilsEditor/Public/Math/Models/CurveTweenModel.h b/Engine/Plugins/Animation/TweeningUtils/Source/TweeningUtilsEditor/Public/Math/Models/CurveTweenModel.h
new file mode 100644
--- /dev/null
+++ b/Engine/Plugins/Animation/TweeningUtils/Source/TweeningUtilsEditor/Public/Math/Models/CurveTweenModel.h
@@ -0,0 +1,38 @@
+﻿// Copyright Epic Games, Inc. All Rights Reserved.
+
+#pragma once
+
+#include "FunctionBasedTweenModel.h"
+#include "Math/ContiguousKeyMapping.h"
+#include "Templates/SharedPointer.h"
+
+class FCurveEditor;
+
+namespace UE::TweeningUtilsEditor
+{
+/** Tween the curves of the curve editor's selection. */
+class TWEENINGUTILSEDITOR_API FCurveTweenModel : public FFunctionBasedTweenModel
+{
+public:
+
+	explicit FCurveTweenModel(TWeakPtr<FCurveEditor> InWeakCurveEditor, const ITweenFunction& InTweenFunction);
+
+	/** @return Whether it makes sense to call StartBlendOperation, etc. */
+	bool HasAnythingToBlend() const;
+	
+	//~ Begin FTweenModel Interface
+	virtual void StartBlendOperation() override;
+	virtual void StopBlendOperation() override;
+	virtual void BlendValues(float InNormalizedValue) override;
+	//~ Begin FTweenModel Interface
+
+protected:
+	
+	/** The curve editor on which to tween the curves. */
+	const TWeakPtr<FCurveEditor> WeakCurveEditor;
+
+	/** Created in StartBlendOperation and used for the entirety of the blend operation. */
+	FContiguousKeyMapping ContiguousKeySelection;
+};
+}
+
diff --git a/Engine/Plugins/Animation/TweeningUtils/Source/TweeningUtilsEditor/Public/Math/Models/EditorTweenModel.h b/Engine/Plugins/Animation/TweeningUtils/Source/TweeningUtilsEditor/Public/Math/Models/EditorTweenModel.h
new file mode 100644
--- /dev/null
+++ b/Engine/Plugins/Animation/TweeningUtils/Source/TweeningUtilsEditor/Public/Math/Models/EditorTweenModel.h
@@ -0,0 +1,31 @@
+﻿// Copyright Epic Games, Inc. All Rights Reserved.
+
+#pragma once
+
+#include "TangentFlatteningTweenProxy.h"
+#include "TransactionalTweenModelProxy.h"
+#include "UserPreferredTweenProxy.h"
+
+#include <type_traits>
+
+namespace UE::TweeningUtilsEditor
+{
+/**
+ * Extends TUnderlyingModel with commonly expected shared functionality:
+ *	- Undo / redo
+ *	- Keeps track of the last selected function and saves it in UTweeningToolsUserSettings.
+ *	- As the tween function squishes the keys vertically, user-specified tangents are flattened.
+ *	- TODO UE-234007: Synchronize tween function selection from different editor tools across all instances of the tweener.
+ *	- In the future this can be extended to remove vertically stacked keys (same keys per frame). For now, no implemented tween function stacks keys.
+ *
+ * Construct like this:
+ *  using FModel = TEditorTweenModel<FFooTweenModel> // FFooTweenModel should inherit from FTweenModel.
+ *	const TSharedRef<FModel> TweenModel = MakeShared<FModel>(
+ *		TEXT("YourFeatureName"), FunctionContainer	// TUserPreferredTweenProxy
+ *		WeakCurveEditor,							// TTangentFlatteningTweenProxy arg
+ *		Bar											// Any arguments that FFooTweenModel requires for construction
+ *	);
+ */
+template<typename TUnderlyingModel> requires std::is_base_of_v<FFunctionBasedTweenModel, TUnderlyingModel> // required by TUserPreferredTweenProxy.
+using TEditorTweenModel = TTransactionalTweenModelProxy<TUserPreferredTweenProxy<TTangentFlatteningTweenProxy<TUnderlyingModel>>>;
+}
diff --git a/Engine/Plugins/Animation/TweeningUtils/Source/TweeningUtilsEditor/Public/Math/Models/FunctionBasedTweenModel.h b/Engine/Plugins/Animation/TweeningUtils/Source/TweeningUtilsEditor/Public/Math/Models/FunctionBasedTweenModel.h
new file mode 100644
--- /dev/null
+++ b/Engine/Plugins/Animation/TweeningUtils/Source/TweeningUtilsEditor/Public/Math/Models/FunctionBasedTweenModel.h
@@ -0,0 +1,31 @@
+﻿// Copyright Epic Games, Inc. All Rights Reserved.
+
+#pragma once
+
+#include "HAL/Platform.h"
+#include "TweenModel.h"
+
+class FCurveEditor;
+
+namespace UE::TweeningUtilsEditor
+{
+class ITweenFunction;
+	
+/** Base class for models that want to implement BlendValues based off of some injectable ITweenFunction. */
+class FFunctionBasedTweenModel : public FTweenModel
+{
+public:
+
+	explicit FFunctionBasedTweenModel(const ITweenFunction& InInitialFunction UE_LIFETIMEBOUND) : TweenFunction(&InInitialFunction) {}
+
+	/** Sets the tween function. You are responsible to make sure it outlives this model. */
+	virtual void SetTweenFunction(const ITweenFunction& InTweenFunction UE_LIFETIMEBOUND) { TweenFunction = &InTweenFunction; }
+	const ITweenFunction& GetTweenFunction() const { return *TweenFunction; }
+	
+private:
+
+	/** The tween function the sub-class will use for blending. You are responsible to make sure it outlives this model. */
+	const ITweenFunction* TweenFunction;
+};
+}
+
diff --git a/Engine/Plugins/Animation/TweeningUtils/Source/TweeningUtilsEditor/Public/Math/Models/TangentFlatteningTweenProxy.h b/Engine/Plugins/Animation/TweeningUtils/Source/TweeningUtilsEditor/Public/Math/Models/TangentFlatteningTweenProxy.h
new file mode 100644
--- /dev/null
+++ b/Engine/Plugins/Animation/TweeningUtils/Source/TweeningUtilsEditor/Public/Math/Models/TangentFlatteningTweenProxy.h
@@ -0,0 +1,58 @@
+﻿// Copyright Epic Games, Inc. All Rights Reserved.
+
+#pragma once
+
+#include "CurveEditor.h"
+#include "Math/Models/TweenModel.h"
+#include "Misc/Attribute.h"
+#include "Misc/Mirror/TangentSelectionFlattener.h"
+#include <type_traits>
+
+namespace UE::TweeningUtilsEditor
+{
+/**
+ * This class squishes the curves based on how much the tween function squishes the keys vertically.
+ * The squishing is achieved by interpolating the tangents to 0.
+ */
+template<typename TBase> requires std::is_base_of_v<FTweenModel, TBase>
+class TTangentFlatteningTweenProxy : public TBase
+{
+public:
+
+	template<typename... TArg>
+	explicit TTangentFlatteningTweenProxy(TAttribute<TWeakPtr<FCurveEditor>> WeakCurveEditorAttr, TArg&&... Arg)
+		: TBase(Forward<TArg>(Arg)...)
+		, WeakCurveEditor(MoveTemp(WeakCurveEditorAttr))
+	{
+		check(WeakCurveEditorAttr.IsBound() || WeakCurveEditorAttr.IsSet());
+	}
+
+	//~ Begin FTweenModel Interface
+	virtual void StartBlendOperation() override
+	{
+		TBase::StartBlendOperation();
+		
+		if (const TSharedPtr<FCurveEditor> CurveEditorPin = WeakCurveEditor.Get().Pin())
+		{
+			TangentTweener.ResetFromSelection(*CurveEditorPin);
+		}
+	}
+	virtual void BlendValues(float InNormalizedValue) override
+	{
+		TBase::BlendValues(InNormalizedValue);
+		if (const TSharedPtr<FCurveEditor> CurveEditorPin = WeakCurveEditor.Get().Pin())
+		{
+			TangentTweener.ComputeMirroringParallel(*CurveEditorPin);
+		}
+	}
+	//~ Begin FTweenModel Interface
+
+private:
+
+	/** Needed as arg for TangentTweener. */
+	const TAttribute<TWeakPtr<FCurveEditor>> WeakCurveEditor;
+	/** Implements the logic for flattening the tangents. */
+	CurveEditor::FTangentSelectionFlattener TangentTweener;
+};
+}
+
diff --git a/Engine/Plugins/Animation/TweeningUtils/Source/TweeningUtilsEditor/Public/Math/Models/TransactionalTweenModelProxy.h b/Engine/Plugins/Animation/TweeningUtils/Source/TweeningUtilsEditor/Public/Math/Models/TransactionalTweenModelProxy.h
new file mode 100644
--- /dev/null
+++ b/Engine/Plugins/Animation/TweeningUtils/Source/TweeningUtilsEditor/Public/Math/Models/TransactionalTweenModelProxy.h
@@ -0,0 +1,46 @@
+﻿// Copyright Epic Games, Inc. All Rights Reserved.
+
+#pragma once
+
+#include "TweenModel.h"
+
+#include "Internationalization/Internationalization.h"
+#include "Internationalization/Text.h"
+#include "ScopedTransaction.h"
+#include "Templates/UniquePtr.h"
+
+#include <type_traits>
+
+class FCurveEditor;
+
+namespace UE::TweeningUtilsEditor
+{
+/** Starts a transaction in StartBlendOperation, stops it in StopBlendOperation.  */
+template<typename TBase> requires std::is_base_of_v<FTweenModel, TBase>
+class TTransactionalTweenModelProxy : public TBase
+{
+public:
+
+	template<typename... TArg>
+	explicit TTransactionalTweenModelProxy(TArg&&... Arg) : TBase(Forward<TArg>(Arg)...) {}
+	
+	//~ Begin FTweenModel Interface
+	virtual void StartBlendOperation() override
+	{
+		InProgressTransaction = MakeUnique<FScopedTransaction>(NSLOCTEXT("FTransactionalTweenModelProxy", "Transaction", "Blend values"));
+		TBase::StartBlendOperation();
+	}
+	virtual void StopBlendOperation() override
+	{
+		TBase::StopBlendOperation();
+		InProgressTransaction.Reset();
+	}
+	//~ Begin FTweenModel Interface
+
+private:
+
+	/** Active during a blend operation. */
+	TUniquePtr<FScopedTransaction> InProgressTransaction;
+};
+}
+
diff --git a/Engine/Plugins/Animation/TweeningUtils/Source/TweeningUtilsEditor/Public/Math/Models/TweenModel.h b/Engine/Plugins/Animation/TweeningUtils/Source/TweeningUtilsEditor/Public/Math/Models/TweenModel.h
new file mode 100644
--- /dev/null
+++ b/Engine/Plugins/Animation/TweeningUtils/Source/TweeningUtilsEditor/Public/Math/Models/TweenModel.h
@@ -0,0 +1,65 @@
+﻿// Copyright Epic Games, Inc. All Rights Reserved.
+
+#pragma once
+
+#include "Misc/Optional.h"
+#include "Widgets/ETweenScaleMode.h"
+
+namespace UE::TweeningUtilsEditor
+{
+/**
+ * Implements the logic of tweening some values: it could be curves, it could be control rig models.
+ * Acts as model in a Model-View-Controller architecture.
+ */
+class TWEENINGUTILSEDITOR_API FTweenModel
+{
+public:
+
+	/** Called when a blend operation is started. */
+	virtual void StartBlendOperation() {}
+	/** Called when a blend operation is stopped. */
+	virtual void StopBlendOperation() {}
+	/**
+	 * Does the blending based on the InNormalizedValue.
+	 * @param InNormalizedValue Value in [-1, 1]
+	 */
+	virtual void BlendValues(float InNormalizedValue) = 0;
+
+	/**
+	 * Blends to a single value once.
+	 * @param InNormalizedValue Value in [-1, 1]
+	 */
+	void BlendOneOff(float InNormalizedValue)
+	{
+		StartBlendOperation();
+		BlendValues(InNormalizedValue);
+		StopBlendOperation();
+	}
+
+	virtual void SetScaleMode(ETweenScaleMode InMode) { ScaleMode = InMode; }
+	ETweenScaleMode GetScaleMode() const { return ScaleMode; }
+	
+	/**
+	 * @param InNormalizedValue Value from -1 to 1.
+	 * @return Value scaled according to the ScaleMode setting.
+	 */
+	float ScaleBlendValue(float InNormalizedValue) const
+	{
+		return ScaleMode == ETweenScaleMode::Normalized
+			? InNormalizedValue
+			: 2.f * InNormalizedValue;
+	}
+	
+	virtual ~FTweenModel() = default;
+	
+private:
+
+	/**
+	 * Affects how normalized values sent to BlendValueNormalizedAndRecomputeKeys are interpreted:
+	 * - False: -1.0 to 1.0 maps to -100% to +100%
+	 * - True: -1.0 to 1.0 maps to -200% to +200%
+	 */
+	ETweenScaleMode ScaleMode = ETweenScaleMode::Normalized;
+};
+}
+
diff --git a/Engine/Plugins/Animation/TweeningUtils/Source/TweeningUtilsEditor/Public/Math/Models/UserPreferredTweenProxy.h b/Engine/Plugins/Animation/TweeningUtils/Source/TweeningUtilsEditor/Public/Math/Models/UserPreferredTweenProxy.h
new file mode 100644
--- /dev/null
+++ b/Engine/Plugins/Animation/TweeningUtils/Source/TweeningUtilsEditor/Public/Math/Models/UserPreferredTweenProxy.h
@@ -0,0 +1,56 @@
+﻿// Copyright Epic Games, Inc. All Rights Reserved.
+
+#pragma once
+
+#include "FunctionBasedTweenModel.h"
+#include "Math/Abstraction/ITweenFunctionContainer.h"
+#include "UObject/NameTypes.h"
+#include "TweeningToolsUserSettings.h"
+#include <type_traits>
+
+namespace UE::TweeningUtilsEditor
+{
+/** Updates the user config with the last selected tween function whenever it is changed.  */
+template<typename TBase> requires std::is_base_of_v<FFunctionBasedTweenModel, TBase>
+class TUserPreferredTweenProxy : public TBase
+{
+public:
+
+	template<typename... TArg>
+	explicit TUserPreferredTweenProxy(FName InFeatureName, const ITweenFunctionContainer& InOwningFunctionContainer UE_LIFETIMEBOUND, TArg&&... Arg)
+		: TBase(Forward<TArg>(Arg)...)
+		, FeatureName(InFeatureName)
+		, OwningFunctionContainer(InOwningFunctionContainer)
+	{
+		// If there's a preferred function saved, load that up as default.
+		const FString* PreferredFunctionId = UTweeningToolsUserSettings::Get()->GetPreferredTweenFunction(FeatureName);
+		const ITweenFunction* PreferredFunction = PreferredFunctionId
+			? OwningFunctionContainer.FindFunctionByIdentifier(*PreferredFunctionId) : nullptr;
+		if (PreferredFunction)
+		{
+			TBase::SetTweenFunction(*PreferredFunction);
+		}
+	}
+
+	//~ Begin FTweenModel Interface
+	virtual void SetTweenFunction(const ITweenFunction& InTweenFunction UE_LIFETIMEBOUND)
+	{
+		const bool bIsContained = OwningFunctionContainer.IsValidIndex(OwningFunctionContainer.IndexOf(InTweenFunction));
+		FString FunctionId = OwningFunctionContainer.GetFunctionIdentifier(InTweenFunction);
+		if (ensureMsgf(bIsContained, TEXT("InTweenFunction is not in ITweenFunctionContainer!")) && !FunctionId.IsEmpty())
+		{
+			UTweeningToolsUserSettings::Get()->SetPreferredTweenFunction(FeatureName, MoveTemp(FunctionId));
+		}
+		
+		TBase::SetTweenFunction(InTweenFunction);
+	}
+	//~ Begin FTweenModel Interface
+
+private:
+
+	/** The key under which the tween function will be saved. */
+	const FName FeatureName;
+	/** The function container from which InTweenFunction originates. You're supposed to ensure that the function container outlives this object.*/
+	const ITweenFunctionContainer& OwningFunctionContainer;
+};
+}
diff --git a/Engine/Plugins/Animation/TweeningUtils/Source/TweeningUtilsEditor/Public/Widgets/MVC/Models/CurveTweenModel.h b/Engine/Plugins/Animation/TweeningUtils/Source/TweeningUtilsEditor/Public/Widgets/MVC/Models/CurveTweenModel.h
deleted file mode 100644
diff --git a/Engine/Plugins/Animation/TweeningUtils/Source/TweeningUtilsEditor/Public/Widgets/MVC/Models/EditorTweenModel.h b/Engine/Plugins/Animation/TweeningUtils/Source/TweeningUtilsEditor/Public/Widgets/MVC/Models/EditorTweenModel.h
deleted file mode 100644
diff --git a/Engine/Plugins/Animation/TweeningUtils/Source/TweeningUtilsEditor/Public/Widgets/MVC/Models/FunctionBasedTweenModel.h b/Engine/Plugins/Animation/TweeningUtils/Source/TweeningUtilsEditor/Public/Widgets/MVC/Models/FunctionBasedTweenModel.h
deleted file mode 100644
diff --git a/Engine/Plugins/Animation/TweeningUtils/Source/TweeningUtilsEditor/Public/Widgets/MVC/Models/TangentFlatteningTweenProxy.h b/Engine/Plugins/Animation/TweeningUtils/Source/TweeningUtilsEditor/Public/Widgets/MVC/Models/TangentFlatteningTweenProxy.h
deleted file mode 100644
diff --git a/Engine/Plugins/Animation/TweeningUtils/Source/TweeningUtilsEditor/Public/Widgets/MVC/Models/TransactionalTweenModelProxy.h b/Engine/Plugins/Animation/TweeningUtils/Source/TweeningUtilsEditor/Public/Widgets/MVC/Models/TransactionalTweenModelProxy.h
deleted file mode 100644
diff --git a/Engine/Plugins/Animation/TweeningUtils/Source/TweeningUtilsEditor/Public/Widgets/MVC/Models/TweenModel.h b/Engine/Plugins/Animation/TweeningUtils/Source/TweeningUtilsEditor/Public/Widgets/MVC/Models/TweenModel.h
deleted file mode 100644
diff --git a/Engine/Plugins/Animation/TweeningUtils/Source/TweeningUtilsEditor/Public/Widgets/MVC/Models/UserPreferredTweenProxy.h b/Engine/Plugins/Animation/TweeningUtils/Source/TweeningUtilsEditor/Public/Widgets/MVC/Models/UserPreferredTweenProxy.h
deleted file mode 100644
diff --git a/Engine/Plugins/Animation/TweeningUtils/Source/TweeningUtilsEditor/TweeningUtilsEditor.build.cs b/Engine/Plugins/Animation/TweeningUtils/Source/TweeningUtilsEditor/TweeningUtilsEditor.build.cs
--- a/Engine/Plugins/Animation/TweeningUtils/Source/TweeningUtilsEditor/TweeningUtilsEditor.build.cs
+++ b/Engine/Plugins/Animation/TweeningUtils/Source/TweeningUtilsEditor/TweeningUtilsEditor.build.cs
@@ -17,6 +17,7 @@ public class TweeningUtilsEditor : ModuleRules
 			"InputCore",
 			"SlateCore",
 			"Slate",
+			"TweeningUtils",
 			"UnrealEd"
 		});
 		
diff --git a/Engine/Plugins/Animation/TweeningUtils/TweeningUtils.uplugin b/Engine/Plugins/Animation/TweeningUtils/TweeningUtils.uplugin
--- a/Engine/Plugins/Animation/TweeningUtils/TweeningUtils.uplugin
+++ b/Engine/Plugins/Animation/TweeningUtils/TweeningUtils.uplugin
@@ -14,6 +14,11 @@
 	"CanContainContent": true,
 	"Installed": false,
 	"Modules": [
+		{
+			"Name": "TweeningUtils",
+			"Type": "Runtime",
+			"LoadingPhase": "Default"
+		},
 		{
 			"Name": "TweeningUtilsEditor",
 			"Type": "Editor",
diff --git a/Engine/Plugins/Editor/CurveEditorTools/Source/CurveEditorTools/Private/Extensions/Tweening/ContextAwareTweeningModel.h b/Engine/Plugins/Editor/CurveEditorTools/Source/CurveEditorTools/Private/Extensions/Tweening/ContextAwareTweeningModel.h
--- a/Engine/Plugins/Editor/CurveEditorTools/Source/CurveEditorTools/Private/Extensions/Tweening/ContextAwareTweeningModel.h
+++ b/Engine/Plugins/Editor/CurveEditorTools/Source/CurveEditorTools/Private/Extensions/Tweening/ContextAwareTweeningModel.h
@@ -3,8 +3,8 @@
 #pragma once
 
 #include "Containers/ContainersFwd.h"
-#include "Widgets/MVC/Models/CurveTweenModel.h"
-#include "Widgets/MVC/Models/FunctionBasedTweenModel.h"
+#include "Math/Models/CurveTweenModel.h"
+#include "Math/Models/FunctionBasedTweenModel.h"
 
 #include <type_traits>
 
diff --git a/Engine/Plugins/Editor/CurveEditorTools/Source/CurveEditorTools/Private/Extensions/Tweening/KeyInsertingTweenModel.h b/Engine/Plugins/Editor/CurveEditorTools/Source/CurveEditorTools/Private/Extensions/Tweening/KeyInsertingTweenModel.h
--- a/Engine/Plugins/Editor/CurveEditorTools/Source/CurveEditorTools/Private/Extensions/Tweening/KeyInsertingTweenModel.h
+++ b/Engine/Plugins/Editor/CurveEditorTools/Source/CurveEditorTools/Private/Extensions/Tweening/KeyInsertingTweenModel.h
@@ -2,8 +2,8 @@
 
 #pragma once
 
+#include "Math/Models/CurveTweenModel.h"
 #include "Templates/SharedPointer.h"
-#include "Widgets/MVC/Models/CurveTweenModel.h"
 
 class FCurveEditor;
 
diff --git a/Engine/Plugins/Editor/CurveEditorTools/Source/CurveEditorTools/Private/Extensions/Tweening/TweenEditorExtension.cpp b/Engine/Plugins/Editor/CurveEditorTools/Source/CurveEditorTools/Private/Extensions/Tweening/TweenEditorExtension.cpp
--- a/Engine/Plugins/Editor/CurveEditorTools/Source/CurveEditorTools/Private/Extensions/Tweening/TweenEditorExtension.cpp
+++ b/Engine/Plugins/Editor/CurveEditorTools/Source/CurveEditorTools/Private/Extensions/Tweening/TweenEditorExtension.cpp
@@ -7,8 +7,8 @@
 #include "Framework/MultiBox/MultiBoxBuilder.h"
 #include "Framework/MultiBox/MultiBoxExtender.h"
 #include "Math/Abstraction/BuiltInTweenFunctionContainer.h"
+#include "Math/Models/EditorTweenModel.h"
 #include "Misc/Optional.h"
-#include "Widgets/MVC/Models/EditorTweenModel.h"
 #include "Widgets/MVC/TweenToolbarController.h"
 
 namespace UE::CurveEditorTools
