Add copyright boilerplate

diff --git a/Engine/Plugins/Animation/TweeningUtils/Source/TweeningUtils/TweeningUtils.Build.cs b/Engine/Plugins/Animation/TweeningUtils/Source/TweeningUtils/TweeningUtils.Build.cs
--- a/Engine/Plugins/Animation/TweeningUtils/Source/TweeningUtils/TweeningUtils.Build.cs
+++ b/Engine/Plugins/Animation/TweeningUtils/Source/TweeningUtils/TweeningUtils.Build.cs
@@ -1,4 +1,6 @@
-﻿using UnrealBuildTool;
+﻿// Copyright Epic Games, Inc. All Rights Reserved.
+
+using UnrealBuildTool;
 
 public class TweeningUtils : ModuleRules
 {
