Curve Editor: Fix crash due to out of bounds index

diff --git a/Engine/Plugins/Animation/TweeningUtils/Source/TweeningUtilsEditor/Private/Math/ContiguousKeyMapping.cpp b/Engine/Plugins/Animation/TweeningUtils/Source/TweeningUtilsEditor/Private/Math/ContiguousKeyMapping.cpp
--- a/Engine/Plugins/Animation/TweeningUtils/Source/TweeningUtilsEditor/Private/Math/ContiguousKeyMapping.cpp
+++ b/Engine/Plugins/Animation/TweeningUtils/Source/TweeningUtilsEditor/Private/Math/ContiguousKeyMapping.cpp
@@ -9,15 +9,16 @@
 #include "CurveEditorSelection.h"
 #include "CurveModel.h"
 #include "Math/NumericLimits.h"
+#include "Misc/Optional.h"
 
 namespace UE::TweeningUtilsEditor
 {
-static int32 GetIndex(const TArray<FKeyPosition>& Keys, double Time)
+static TOptional<int32> GetIndex(const TArray<FKeyPosition>& Keys, double Time)
 {
 	const int32 Index = Algo::LowerBoundBy(Keys, Time, [](const FKeyPosition& Value) { return Value.InputValue; });
 
 	// don't trust precision issues so will double check to make sure the index is correct
-	if (Index != INDEX_NONE)
+	if (Keys.IsValidIndex(Index))
 	{
 		if (FMath::IsNearlyEqual(Keys[Index].InputValue, Time))
 		{
@@ -31,8 +32,9 @@ static int32 GetIndex(const TArray<FKeyPosition>& Keys, double Time)
 		{
 			return Index + 1;
 		}
+		return Index;
 	}
-	return Index;
+	return {};
 }
 
 static void AppendKeyArray(
@@ -70,8 +72,11 @@ static void AppendKeyArray(
 		TArray<int32> SelectedIndices;
 		for (int32 Index = 0; Index < KeyPositions.Num(); ++Index)
 		{
-			int32 AllKeysIndex = GetIndex(AllKeyPositions, KeyPositions[Index].InputValue);
-			SelectedIndices.Add(AllKeysIndex);
+			const TOptional<int32> AllKeysIndex = GetIndex(AllKeyPositions, KeyPositions[Index].InputValue);
+			if (AllKeysIndex)
+			{
+				SelectedIndices.Add(*AllKeysIndex);
+			}
 		}
 		SelectedIndices.Sort();
 		
