TweenUtils: Refactor ITweenFunctionContainer to ITweenModelContainer. This allows the slider's combo button to select between FTweenModels instead of just tween functions. This allows for more flexibility.

diff --git a/Engine/Plugins/Animation/ControlRig/Source/ControlRigEditor/Private/ControlRigSequencerEditorLibrary.cpp b/Engine/Plugins/Animation/ControlRig/Source/ControlRigEditor/Private/ControlRigSequencerEditorLibrary.cpp
--- a/Engine/Plugins/Animation/ControlRig/Source/ControlRigEditor/Private/ControlRigSequencerEditorLibrary.cpp
+++ b/Engine/Plugins/Animation/ControlRig/Source/ControlRigEditor/Private/ControlRigSequencerEditorLibrary.cpp
@@ -51,7 +51,6 @@
 #include "EditMode/ControlRigEditMode.h"
 #include "SequencerTools.h"
 #include "EditMode/Tween/ControlRigTweenModel.h"
-#include "Math/Abstraction/BuiltInTweenFunctionContainer.h"
 #include "Math/Abstraction/KeyBlendingAbstraction.h"
 #include "Sequencer/AnimLayers/AnimLayers.h"
 
@@ -411,8 +410,7 @@ bool UControlRigSequencerEditorLibrary::TweenControlRig(ULevelSequence* LevelSeq
 	{
 		LevelSequence->GetMovieScene()->Modify();
 		
-		const FBuiltInTweenFunctionContainer TweenFunctionContainer; 
-		FControlRigTweenModel BlendTool(WeakSequencer, *TweenFunctionContainer.GetBlendFunction(EBlendFunction::ControlsToTween));
+		FControlRigTweenModel BlendTool(WeakSequencer, EBlendFunction::ControlsToTween);
 		BlendTool.StartBlendOperationWithControlRigs({ ControlRig });
 		BlendTool.BlendValues(TweenValue);
 		BlendTool.StopBlendOperation();
@@ -452,11 +450,8 @@ bool UControlRigSequencerEditorLibrary::BlendValuesOnSelected(ULevelSequence* Le
 			TWeakPtr<FControlRigEditMode> WeakMode = StaticCastSharedRef<FControlRigEditMode, FEdMode>(EditMode->AsShared()).ToWeakPtr();
 			LevelSequence->GetMovieScene()->Modify();
 			
-			const FBuiltInTweenFunctionContainer TweenFunctionContainer;
 			const EBlendFunction BlendFunction = SequencerLibraryDetail::ConvertBlendFunction(BlendOperation);
-			FControlRigTweenModel BlendTool(WeakSequencer, WeakMode,
-				*TweenFunctionContainer.GetBlendFunction(BlendFunction)
-				);
+			FControlRigTweenModel BlendTool(WeakSequencer, WeakMode, BlendFunction);
 			return BlendTool.BlendSingle(BlendValue);
 		}
 	}
diff --git a/Engine/Plugins/Animation/ControlRig/Source/ControlRigEditor/Private/EditMode/Tween/ControlRigTweenModel.cpp b/Engine/Plugins/Animation/ControlRig/Source/ControlRigEditor/Private/EditMode/Tween/ControlRigTweenModel.cpp
--- a/Engine/Plugins/Animation/ControlRig/Source/ControlRigEditor/Private/EditMode/Tween/ControlRigTweenModel.cpp
+++ b/Engine/Plugins/Animation/ControlRig/Source/ControlRigEditor/Private/EditMode/Tween/ControlRigTweenModel.cpp
@@ -5,7 +5,6 @@
 #include "ISequencer.h"
 #include "Channels/MovieSceneFloatChannel.h"
 #include "Channels/MovieSceneDoubleChannel.h"
-#include "Math/Abstraction/ITweenFunction.h"
 #include "Math/CurveBlending.h"
 #include "Misc/SequencerUtils.h"
 
@@ -71,7 +70,7 @@ bool FControlRigTweenModel::DoBlend(float InNormalizedValue)
 	{
 		return TweeningUtilsEditor::BlendCurves(ScaledBlendValue, *CurveEditor, KeySelection, [this](const TweeningUtils::FKeyBlendInfo& InBlendInfo)
 		{
-			return GetTweenFunction().Blend(InBlendInfo);
+			return EvaluateBlend(BlendFunction, InBlendInfo);
 		});
 	}
 
@@ -110,7 +109,8 @@ bool FControlRigTweenModel::DoBlend(float InNormalizedValue)
 						Current, AfterBlendRange, // Current is also the last key to blend
 						ScaledBlendValue
 						);
-					const double NewValue = GetTweenFunction().Blend(BlendStruct);
+					const double NewValue = EvaluateBlend(BlendFunction, BlendStruct);
+					
 					using namespace UE::MovieScene;
 					if (ObjectChannels.KeyBounds[Index].FloatChannel)
 					{
diff --git a/Engine/Plugins/Animation/ControlRig/Source/ControlRigEditor/Private/EditMode/Tween/ControlRigTweenModel.h b/Engine/Plugins/Animation/ControlRig/Source/ControlRigEditor/Private/EditMode/Tween/ControlRigTweenModel.h
--- a/Engine/Plugins/Animation/ControlRig/Source/ControlRigEditor/Private/EditMode/Tween/ControlRigTweenModel.h
+++ b/Engine/Plugins/Animation/ControlRig/Source/ControlRigEditor/Private/EditMode/Tween/ControlRigTweenModel.h
@@ -4,7 +4,8 @@
 
 #include "ControlRigObjectSelection.h"
 #include "Math/ContiguousKeyMapping.h"
-#include "Math/Models/FunctionBasedTweenModel.h"
+#include "Math/Abstraction/KeyBlendingAbstraction.h"
+#include "Math/Models/TweenModel.h"
 #include "Misc/Attribute.h"
 #include "Templates/SharedPointer.h"
 
@@ -20,23 +21,20 @@ namespace UE::ControlRigEditor
  *	2. Otherwise, if a control rig element is selected, insert a key at the current scrub position into the curves corresponding to the control rig
  *	and then interpolate those keys.
  */
-class FControlRigTweenModel : public TweeningUtilsEditor::FFunctionBasedTweenModel
+class FControlRigTweenModel : public TweeningUtilsEditor::FTweenModel
 {
 public:
 	
 	explicit FControlRigTweenModel(
 		TAttribute<TWeakPtr<ISequencer>> InSequencerAttr, TWeakPtr<FControlRigEditMode> InControlRigMode,
-		const TweeningUtilsEditor::ITweenFunction& InInitialFunction
+		TweeningUtilsEditor::EBlendFunction InBlendFunction
 		)
-		: FFunctionBasedTweenModel(InInitialFunction)
-		, SequencerAttr(MoveTemp(InSequencerAttr))
+		: SequencerAttr(MoveTemp(InSequencerAttr))
 		, ControlRigMode(MoveTemp(InControlRigMode))
+		, BlendFunction(InBlendFunction)
 	{}
-	explicit FControlRigTweenModel(
-		TAttribute<TWeakPtr<ISequencer>> InSequencerAttr, const TweeningUtilsEditor::ITweenFunction& InInitialFunction
-		)
-		: FFunctionBasedTweenModel(InInitialFunction)
-		, SequencerAttr(MoveTemp(InSequencerAttr))
+	explicit FControlRigTweenModel(TAttribute<TWeakPtr<ISequencer>> InSequencerAttr, TweeningUtilsEditor::EBlendFunction InBlendFunction)
+		: FControlRigTweenModel(MoveTemp(InSequencerAttr), nullptr, InBlendFunction)
 	{}
 
 	/** Does a blend and returns whether any values were actually blended. */
@@ -58,6 +56,8 @@ private:
 	const TAttribute<TWeakPtr<ISequencer>> SequencerAttr;
 	/** The control rig mode that owns this tweener. Can be nullptr, in which case only the curve editor key selection is blended. */
 	const TWeakPtr<FControlRigEditMode> ControlRigMode;
+	/** The blend function that does the blending. */
+	const TweeningUtilsEditor::EBlendFunction BlendFunction;
 
 	/** The key selection to blend. Created in StartBlendOperation and used for the entirety of the blend operation. */
 	TweeningUtilsEditor::FContiguousKeyMapping KeySelection;
@@ -65,7 +65,7 @@ private:
 	 * The control rig selection to blend if ContiguousKeySelection is empty.
 	 * Created in StartBlendOperation and used for the entirety of the blend operation.
 	 */
-	FControlRigObjectSelection ObjectSelection; // TODO: Extract this to UE::ControlRigEditor namespace and deprecated the global one.
+	FControlRigObjectSelection ObjectSelection; // TODO: Extract this to UE::ControlRigEditor namespace and deprecate the global one.
 	
 	/** @return Whether any blended occured. */
 	bool DoBlend(float InNormalizedValue);
diff --git a/Engine/Plugins/Animation/ControlRig/Source/ControlRigEditor/Private/EditMode/Tween/TweenLogic.cpp b/Engine/Plugins/Animation/ControlRig/Source/ControlRigEditor/Private/EditMode/Tween/TweenLogic.cpp
--- a/Engine/Plugins/Animation/ControlRig/Source/ControlRigEditor/Private/EditMode/Tween/TweenLogic.cpp
+++ b/Engine/Plugins/Animation/ControlRig/Source/ControlRigEditor/Private/EditMode/Tween/TweenLogic.cpp
@@ -5,7 +5,7 @@
 #include "ControlRigTweenModel.h"
 #include "EditMode/ControlRigEditMode.h"
 #include "Framework/Commands/UICommandList.h"
-#include "Math/Abstraction/BuiltInTweenFunctionContainer.h"
+#include "Math/Abstraction/TweenModelArray.h"
 #include "Math/Models/EditorTweenModel.h"
 #include "Misc/SequencerUtils.h"
 #include "Templates/SharedPointer.h"
@@ -14,35 +14,33 @@ namespace UE::ControlRigEditor
 {
 namespace TweenLogicDetail
 {
-static TSharedRef<FControlRigTweenModel> MakeTweenModel(
+static TUniquePtr<FControlRigTweenModel> MakeTweenModel(
 	TAttribute<TWeakPtr<ISequencer>> InSequencerAttr, const TSharedRef<FControlRigEditMode>& InOwningEditMode,
-	const TweeningUtilsEditor::ITweenFunctionContainer& InFunctions
+	TweeningUtilsEditor::EBlendFunction BlendFunction
 	)
 {
 	using FModel = TweeningUtilsEditor::TEditorTweenModel<FControlRigTweenModel>;
-	
-	const FName FeatureName = TEXT("ControlRigTween"); // The key in UTweeningToolsUserSettings from which to look up the preferred function.
-	const TweeningUtilsEditor::ITweenFunction& DefaultFunction = *InFunctions.GetFunction(0); // Use 1st function if no preferred function is in config.
-	
 	TAttribute<TWeakPtr<FCurveEditor>> CurveEditorAttr = TAttribute<TWeakPtr<FCurveEditor>>::CreateLambda([SequencerAttr = InSequencerAttr]
 	{
 		const TSharedPtr<ISequencer> SequencerPin = SequencerAttr.Get().Pin();
 		return GetCurveEditorFromSequencer(SequencerPin);
 	});
 	
-	return MakeShared<FModel>(
-		FeatureName, InFunctions,										// TUserPreferredTweenProxy
+	return MakeUnique<FModel>(
 		MoveTemp(CurveEditorAttr),										// TTangentFlatteningTweenProxy arg
-		MoveTemp(InSequencerAttr), InOwningEditMode, DefaultFunction
+		MoveTemp(InSequencerAttr), InOwningEditMode, BlendFunction
 		);
 }
 }
 	
 FTweenLogic::FTweenLogic(TAttribute<TWeakPtr<ISequencer>> InSequencerAttr, const TSharedRef<FControlRigEditMode>& InOwningEditMode)
 	: CommandList(InOwningEditMode->GetCommandBindings().ToSharedRef())
-	, TweenFunctions(MakeShared<TweeningUtilsEditor::FBuiltInTweenFunctionContainer>())
-	, TweenModel(TweenLogicDetail::MakeTweenModel(MoveTemp(InSequencerAttr), InOwningEditMode, *TweenFunctions)) 
-	, Controllers(TweenModel, InOwningEditMode->GetCommandBindings().ToSharedRef(), TweenFunctions)
+	, TweenModels(MakeShared<TweeningUtilsEditor::FTweenModelArray>(
+		[&InSequencerAttr, &InOwningEditMode](TweeningUtilsEditor::EBlendFunction BlendFunction)
+		{
+			return TweenLogicDetail::MakeTweenModel(InSequencerAttr, InOwningEditMode, BlendFunction);
+		}))
+	, Controllers(CommandList, TweenModels, TEXT("ControlRigTween"))
 {}
 
 TSharedRef<SWidget> FTweenLogic::MakeWidget() const
diff --git a/Engine/Plugins/Animation/ControlRig/Source/ControlRigEditor/Private/EditMode/Tween/TweenLogic.h b/Engine/Plugins/Animation/ControlRig/Source/ControlRigEditor/Private/EditMode/Tween/TweenLogic.h
--- a/Engine/Plugins/Animation/ControlRig/Source/ControlRigEditor/Private/EditMode/Tween/TweenLogic.h
+++ b/Engine/Plugins/Animation/ControlRig/Source/ControlRigEditor/Private/EditMode/Tween/TweenLogic.h
@@ -9,7 +9,6 @@
 class FControlRigEditMode;
 class FUICommandList;
 class ISequencer;
-namespace UE::TweeningUtilsEditor { class ITweenFunctionContainer; }
 
 namespace UE::ControlRigEditor
 {
@@ -31,9 +30,7 @@ private:
 	const TSharedRef<FUICommandList> CommandList;
 
 	/** Holds the used tweening functions and info about how to display them in UI. */
-	const TSharedRef<TweeningUtilsEditor::ITweenFunctionContainer> TweenFunctions;
-	/** The control rig specific tween model. */
-	const TSharedRef<FControlRigTweenModel> TweenModel;
+	const TSharedRef<TweeningUtilsEditor::ITweenModelContainer> TweenModels;
 
 	/** Common functionality that should be shared consistent with other editor modules, such as Curve Editor. */
 	const TweeningUtilsEditor::FTweenControllers Controllers;
diff --git a/Engine/Plugins/Animation/TweeningUtils/Source/TweeningUtilsEditor/Private/Math/Abstraction/BuiltInTweenFunctionContainer.cpp b/Engine/Plugins/Animation/TweeningUtils/Source/TweeningUtilsEditor/Private/Math/Abstraction/BuiltInTweenFunctionContainer.cpp
deleted file mode 100644
diff --git a/Engine/Plugins/Animation/TweeningUtils/Source/TweeningUtilsEditor/Private/Math/Abstraction/BuiltInTweenModelContainer.cpp b/Engine/Plugins/Animation/TweeningUtils/Source/TweeningUtilsEditor/Private/Math/Abstraction/BuiltInTweenModelContainer.cpp
new file mode 100644
--- /dev/null
+++ b/Engine/Plugins/Animation/TweeningUtils/Source/TweeningUtilsEditor/Private/Math/Abstraction/BuiltInTweenModelContainer.cpp
@@ -0,0 +1,82 @@
+﻿// Copyright Epic Games, Inc. All Rights Reserved.
+
+#include "Math/Abstraction/BuiltInTweenModelContainer.h"
+
+#include "Containers/UnrealString.h"
+#include "Framework/Commands/UICommandInfo.h"
+#include "Misc/Optional.h"
+#include "Templates/Function.h"
+
+namespace UE::TweeningUtilsEditor
+{
+TOptional<EBlendFunction> FBuiltInTweenModelContainer::GetBlendFunctionType(const FTweenModel& InTweenModelToFind) const
+{
+	for (int32 Index = 0; Index < Functions.size(); ++Index)
+	{
+		const TUniquePtr<FTweenModel>& TweenModel = Functions[Index];
+		if (&InTweenModelToFind == TweenModel.Get())
+		{
+			return static_cast<EBlendFunction>(Index);
+		}
+	}
+	return {};
+}
+
+const FTweenModel* FBuiltInTweenModelContainer::GetBlendFunction(EBlendFunction InBlendFunction) const
+{
+	const int32 Index = static_cast<int32>(InBlendFunction);
+	return ensure(Index >= 0 && Index < Functions.size()) ? Functions[Index].Get() : nullptr;
+}
+
+void FBuiltInTweenModelContainer::ForEachModel(const TFunctionRef<void(FTweenModel&)> InConsumer) const
+{
+	for (const TUniquePtr<FTweenModel>& TweenModel : Functions)
+	{
+		if (TweenModel) // Handle that case that the API user that constructed this FBuiltInTweenModelContainer, did not specify this built-in function.
+		{
+			InConsumer(*TweenModel.Get());
+		}
+	}
+}
+
+FTweenModel* FBuiltInTweenModelContainer::GetModel(int32 InIndex) const
+{
+	return Functions[InIndex].Get();
+}
+
+TSharedPtr<FUICommandInfo> FBuiltInTweenModelContainer::GetCommandForModel(const FTweenModel& InFunction) const
+{
+	const TOptional<EBlendFunction> Function = GetBlendFunctionType(InFunction);
+	return Function ? GetCommandForBlendFunction(*Function) : nullptr;
+}
+
+const FSlateBrush* FBuiltInTweenModelContainer::GetIconForModel(const FTweenModel& InFunction) const
+{
+	const TOptional<EBlendFunction> Function = GetBlendFunctionType(InFunction);
+	return Function ? GetUntintedIconForTweenFunction(*Function) : nullptr;
+}
+
+FLinearColor FBuiltInTweenModelContainer::GetColorForModel(const FTweenModel& InFunction) const
+{
+	const TOptional<EBlendFunction> Function = GetBlendFunctionType(InFunction);
+	return Function ? GetTintColorForTweenFunction(*Function) : FLinearColor::White;
+}
+
+FText FBuiltInTweenModelContainer::GetLabelForModel(const FTweenModel& InFunction) const
+{
+	const TOptional<EBlendFunction> Function = GetBlendFunctionType(InFunction);
+	return Function ? GetLabelForBlendFunction(*Function) : FText::GetEmpty();
+}
+	
+FText FBuiltInTweenModelContainer::GetToolTipForModel(const FTweenModel& InFunction) const
+{
+	const TOptional<EBlendFunction> Function = GetBlendFunctionType(InFunction);
+	return Function ? GetDescriptionForBlendFunction(*Function) : FText::GetEmpty();
+}
+
+FString FBuiltInTweenModelContainer::GetModelIdentifier(const FTweenModel& InFunction) const
+{
+	const TOptional<EBlendFunction> Function = GetBlendFunctionType(InFunction);
+	return Function ? BlendFunctionToString(*Function) : FString();
+}
+}
diff --git a/Engine/Plugins/Animation/TweeningUtils/Source/TweeningUtilsEditor/Private/Math/Abstraction/ITweenFunctionContainer.cpp b/Engine/Plugins/Animation/TweeningUtils/Source/TweeningUtilsEditor/Private/Math/Abstraction/ITweenFunctionContainer.cpp
deleted file mode 100644
diff --git a/Engine/Plugins/Animation/TweeningUtils/Source/TweeningUtilsEditor/Private/Math/Abstraction/ITweenModelContainer.cpp b/Engine/Plugins/Animation/TweeningUtils/Source/TweeningUtilsEditor/Private/Math/Abstraction/ITweenModelContainer.cpp
new file mode 100644
--- /dev/null
+++ b/Engine/Plugins/Animation/TweeningUtils/Source/TweeningUtilsEditor/Private/Math/Abstraction/ITweenModelContainer.cpp
@@ -0,0 +1,34 @@
+﻿// Copyright Epic Games, Inc. All Rights Reserved.
+
+#include "Math/Abstraction/ITweenModelContainer.h"
+
+#include "Containers/UnrealString.h"
+#include "Misc/CoreMiscDefines.h"
+
+namespace UE::TweeningUtilsEditor
+{
+int32 ITweenModelContainer::IndexOf(const FTweenModel& InTweenModel) const
+{
+	for (int32 Index = 0; Index < NumModels(); ++Index)
+	{
+		if (GetModel(Index) == &InTweenModel)
+		{
+			return Index;
+		}
+	}
+	return INDEX_NONE;
+}
+
+FTweenModel* ITweenModelContainer::FindModelByIdentifier(const FString& InIdentifier)
+{
+	for (int32 Index = 0; Index < NumModels(); ++Index)
+	{
+		FTweenModel* Function = GetModel(Index);
+		if (Function && GetModelIdentifier(*Function) == InIdentifier)
+		{
+			return Function;
+		}
+	}
+	return nullptr;
+}
+}
diff --git a/Engine/Plugins/Animation/TweeningUtils/Source/TweeningUtilsEditor/Private/Math/Abstraction/TweenModelArray.cpp b/Engine/Plugins/Animation/TweeningUtils/Source/TweeningUtilsEditor/Private/Math/Abstraction/TweenModelArray.cpp
new file mode 100644
--- /dev/null
+++ b/Engine/Plugins/Animation/TweeningUtils/Source/TweeningUtilsEditor/Private/Math/Abstraction/TweenModelArray.cpp
@@ -0,0 +1,82 @@
+﻿// Copyright Epic Games, Inc. All Rights Reserved.
+
+#include "Math/Abstraction/TweenModelArray.h"
+
+#include "Containers/UnrealString.h"
+#include "Framework/Commands/UICommandInfo.h"
+#include "Misc/Optional.h"
+#include "Templates/Function.h"
+
+namespace UE::TweeningUtilsEditor
+{
+TOptional<EBlendFunction> FTweenModelArray::GetBlendFunctionType(const FTweenModel& InTweenModelToFind) const
+{
+	for (int32 Index = 0; Index < Models.size(); ++Index)
+	{
+		const TUniquePtr<FTweenModel>& TweenModel = Models[Index];
+		if (&InTweenModelToFind == TweenModel.Get())
+		{
+			return static_cast<EBlendFunction>(Index);
+		}
+	}
+	return {};
+}
+
+const FTweenModel* FTweenModelArray::GetBlendFunction(EBlendFunction InBlendFunction) const
+{
+	const int32 Index = static_cast<int32>(InBlendFunction);
+	return ensure(Index >= 0 && Index < Models.size()) ? Models[Index].Get() : nullptr;
+}
+
+void FTweenModelArray::ForEachModel(const TFunctionRef<void(FTweenModel&)> InConsumer) const
+{
+	for (const TUniquePtr<FTweenModel>& TweenModel : Models)
+	{
+		if (TweenModel) // Handle that case that the API user that constructed this FBuiltInTweenModelContainer, did not specify this built-in function.
+		{
+			InConsumer(*TweenModel.Get());
+		}
+	}
+}
+
+FTweenModel* FTweenModelArray::GetModel(int32 InIndex) const
+{
+	return Models[InIndex].Get();
+}
+
+TSharedPtr<FUICommandInfo> FTweenModelArray::GetCommandForModel(const FTweenModel& InFunction) const
+{
+	const TOptional<EBlendFunction> Function = GetBlendFunctionType(InFunction);
+	return Function ? GetCommandForBlendFunction(*Function) : nullptr;
+}
+
+const FSlateBrush* FTweenModelArray::GetIconForModel(const FTweenModel& InFunction) const
+{
+	const TOptional<EBlendFunction> Function = GetBlendFunctionType(InFunction);
+	return Function ? GetUntintedIconForTweenFunction(*Function) : nullptr;
+}
+
+FLinearColor FTweenModelArray::GetColorForModel(const FTweenModel& InFunction) const
+{
+	const TOptional<EBlendFunction> Function = GetBlendFunctionType(InFunction);
+	return Function ? GetTintColorForTweenFunction(*Function) : FLinearColor::White;
+}
+
+FText FTweenModelArray::GetLabelForModel(const FTweenModel& InFunction) const
+{
+	const TOptional<EBlendFunction> Function = GetBlendFunctionType(InFunction);
+	return Function ? GetLabelForBlendFunction(*Function) : FText::GetEmpty();
+}
+	
+FText FTweenModelArray::GetToolTipForModel(const FTweenModel& InFunction) const
+{
+	const TOptional<EBlendFunction> Function = GetBlendFunctionType(InFunction);
+	return Function ? GetDescriptionForBlendFunction(*Function) : FText::GetEmpty();
+}
+
+FString FTweenModelArray::GetModelIdentifier(const FTweenModel& InFunction) const
+{
+	const TOptional<EBlendFunction> Function = GetBlendFunctionType(InFunction);
+	return Function ? BlendFunctionToString(*Function) : FString();
+}
+}
diff --git a/Engine/Plugins/Animation/TweeningUtils/Source/TweeningUtilsEditor/Private/Math/Models/CurveTweeningModel.cpp b/Engine/Plugins/Animation/TweeningUtils/Source/TweeningUtilsEditor/Private/Math/Models/CurveTweeningModel.cpp
--- a/Engine/Plugins/Animation/TweeningUtils/Source/TweeningUtilsEditor/Private/Math/Models/CurveTweeningModel.cpp
+++ b/Engine/Plugins/Animation/TweeningUtils/Source/TweeningUtilsEditor/Private/Math/Models/CurveTweeningModel.cpp
@@ -3,13 +3,12 @@
 #include "Math/Models/CurveTweenModel.h"
 
 #include "Math/CurveBlending.h"
-#include "Math/Abstraction/ITweenFunction.h"
 
 namespace UE::TweeningUtilsEditor
 {
-FCurveTweenModel::FCurveTweenModel(TWeakPtr<FCurveEditor> InWeakCurveEditor, const ITweenFunction& InTweenFunction)
-	: FFunctionBasedTweenModel(InTweenFunction)
-	, WeakCurveEditor(InWeakCurveEditor)
+FCurveTweenModel::FCurveTweenModel(TWeakPtr<FCurveEditor> InWeakCurveEditor, EBlendFunction InBlendFunction)
+	: WeakCurveEditor(MoveTemp(InWeakCurveEditor))
+	, BlendFunction(InBlendFunction)
 {}
 
 bool FCurveTweenModel::HasAnythingToBlend() const
@@ -42,7 +41,7 @@ void FCurveTweenModel::BlendValues(float InNormalizedValue)
 	const float ScaledBlendValue = ScaleBlendValue(InNormalizedValue); 
 	TweeningUtilsEditor::BlendCurves(ScaledBlendValue, *CurveEditor, ContiguousKeySelection, [this](const TweeningUtils::FKeyBlendInfo& InBlendInfo)
 	{
-		return GetTweenFunction().Blend(InBlendInfo);
+		return EvaluateBlend(BlendFunction, InBlendInfo);
 	});
 }
 }
diff --git a/Engine/Plugins/Animation/TweeningUtils/Source/TweeningUtilsEditor/Private/Widgets/MVC/CycleFunctionController.cpp b/Engine/Plugins/Animation/TweeningUtils/Source/TweeningUtilsEditor/Private/Widgets/MVC/CycleFunctionController.cpp
--- a/Engine/Plugins/Animation/TweeningUtils/Source/TweeningUtilsEditor/Private/Widgets/MVC/CycleFunctionController.cpp
+++ b/Engine/Plugins/Animation/TweeningUtils/Source/TweeningUtilsEditor/Private/Widgets/MVC/CycleFunctionController.cpp
@@ -4,33 +4,39 @@
 
 #include "Framework/Commands/UICommandInfo.h"
 #include "Framework/Commands/UICommandList.h"
-#include "Math/Abstraction/ITweenFunctionContainer.h"
-#include "Math/Models/FunctionBasedTweenModel.h"
+#include "Math/Abstraction/ITweenModelContainer.h"
 #include "TweeningUtilsCommands.h"
 
 namespace UE::TweeningUtilsEditor
 {
 FCycleFunctionController::FCycleFunctionController(
-	const TSharedRef<FFunctionBasedTweenModel>& InTweenModel,
-	const TSharedRef<ITweenFunctionContainer>& InTweenFunctions,
+	TAttribute<const FTweenModel*> InCurrentTweenModelAttr,
+	const TSharedRef<ITweenModelContainer>& InTweenModelContainer,
+	FHandleTweenChange InHandleTweenChange,
 	const TSharedRef<FUICommandList>& InCommandList,
 	TSharedPtr<FUICommandInfo> InCycleCommand
 	)
-	: TweenModel(InTweenModel)
-	, TweenFunctions(InTweenFunctions)
+	: CurrentTweenModelAttr(MoveTemp(InCurrentTweenModelAttr))
+	, TweenModelContainer(InTweenModelContainer)
+	, HandleTweenChangeDelegate(MoveTemp(InHandleTweenChange))
 	, CommandList(InCommandList)
 	, CycleCommand(MoveTemp(InCycleCommand))
 {
+	check(CurrentTweenModelAttr.IsSet() || CurrentTweenModelAttr.IsBound());
 	check(CycleCommand);
 	CommandList->MapAction(CycleCommand, FExecuteAction::CreateRaw(this, &FCycleFunctionController::CycleToNextFunction));
 }
 
 FCycleFunctionController::FCycleFunctionController(
-	const TSharedRef<FFunctionBasedTweenModel>& InTweenModel,
-	const TSharedRef<ITweenFunctionContainer>& InTweenFunctions,
+	TAttribute<const FTweenModel*> InCurrentTweenModelAttr,
+	const TSharedRef<ITweenModelContainer>& InTweenModelContainer,
+	FHandleTweenChange InHandleTweenChange,
 	const TSharedRef<FUICommandList>& InCommandList
 	)
-	: FCycleFunctionController(InTweenModel, InTweenFunctions, InCommandList, FTweeningUtilsCommands::Get().ChangeAnimSliderTool)
+	: FCycleFunctionController(
+		MoveTemp(InCurrentTweenModelAttr), InTweenModelContainer, MoveTemp(InHandleTweenChange),
+		InCommandList, FTweeningUtilsCommands::Get().ChangeAnimSliderTool
+		)
 {}
 
 FCycleFunctionController::~FCycleFunctionController()
@@ -40,11 +46,11 @@ FCycleFunctionController::~FCycleFunctionController()
 
 void FCycleFunctionController::CycleToNextFunction()
 {
-	const int32 Index = TweenFunctions->IndexOf(TweenModel->GetTweenFunction());
-	if (ensure(TweenFunctions->IsValidIndex(Index)))
+	const int32 Index = TweenModelContainer->IndexOf(*CurrentTweenModelAttr.Get());
+	if (ensure(TweenModelContainer->IsValidIndex(Index)))
 	{
-		const int32 NextIndex = TweenFunctions->IsValidIndex(Index + 1) ? Index + 1 : 0;
-		TweenModel->SetTweenFunction(*TweenFunctions->GetFunction(NextIndex));
+		const int32 NextIndex = TweenModelContainer->IsValidIndex(Index + 1) ? Index + 1 : 0;
+		HandleTweenChangeDelegate.Execute(*TweenModelContainer->GetModel(NextIndex));
 	}
 }
 }
diff --git a/Engine/Plugins/Animation/TweeningUtils/Source/TweeningUtilsEditor/Private/Widgets/MVC/STweenView.cpp b/Engine/Plugins/Animation/TweeningUtils/Source/TweeningUtilsEditor/Private/Widgets/MVC/STweenView.cpp
--- a/Engine/Plugins/Animation/TweeningUtils/Source/TweeningUtilsEditor/Private/Widgets/MVC/STweenView.cpp
+++ b/Engine/Plugins/Animation/TweeningUtils/Source/TweeningUtilsEditor/Private/Widgets/MVC/STweenView.cpp
@@ -9,12 +9,13 @@
 
 namespace UE::TweeningUtilsEditor
 {
-void STweenView::Construct(const FArguments& InArgs, const TSharedRef<FTweenModel>& InTweenModel)
+void STweenView::Construct(const FArguments& InArgs)
 {
+	check(InArgs._TweenModel.IsSet() || InArgs._TweenModel.IsBound());
 	check(InArgs._SliderIcon.IsSet() || InArgs._SliderIcon.IsBound());
 	check(InArgs._SliderColor.IsSet() || InArgs._SliderColor.IsBound());
 	
-	TweenModel = InTweenModel;
+	TweenModelAttr = InArgs._TweenModel;
 
 	// STweenView is placed in toolbars.
 	// We want STweenSlider's foreground color to be the same as that of the SComboButton from which you pick functions for color consistency.
@@ -43,31 +44,31 @@ void STweenView::Construct(const FArguments& InArgs, const TSharedRef<FTweenMode
 
 ETweenScaleMode STweenView::GetBarRenderMode() const
 {
-	return TweenModel->GetScaleMode();
+	return TweenModelAttr.Get()->GetScaleMode();
 }
 
-void STweenView::OnDragStarted()
+void STweenView::OnDragStarted() const
 {
-	TweenModel->StartBlendOperation();
+	TweenModelAttr.Get()->StartBlendOperation();
 }
 
-void STweenView::OnDragEnded()
+void STweenView::OnDragEnded() const
 {
-	TweenModel->StopBlendOperation();
+	TweenModelAttr.Get()->StopBlendOperation();
 }
 
-void STweenView::OnDragValueUpdated(float Value)
+void STweenView::OnDragValueUpdated(float Value) const
 {
-	TweenModel->BlendValues(Value);
+	TweenModelAttr.Get()->BlendValues(Value);
 }
 
-void STweenView::OnPointPicked(float Value)
+void STweenView::OnPointPicked(float Value) const
 {
-	TweenModel->BlendOneOff(Value);
+	TweenModelAttr.Get()->BlendOneOff(Value);
 }
 
-float STweenView::MapSliderValueToBlendValue(float Value)
+float STweenView::MapSliderValueToBlendValue(float Value) const
 {
-	return TweenModel->ScaleBlendValue(Value);
+	return TweenModelAttr.Get()->ScaleBlendValue(Value);
 }
 }
diff --git a/Engine/Plugins/Animation/TweeningUtils/Source/TweeningUtilsEditor/Private/Widgets/MVC/TweenControllers.cpp b/Engine/Plugins/Animation/TweeningUtils/Source/TweeningUtilsEditor/Private/Widgets/MVC/TweenControllers.cpp
--- a/Engine/Plugins/Animation/TweeningUtils/Source/TweeningUtilsEditor/Private/Widgets/MVC/TweenControllers.cpp
+++ b/Engine/Plugins/Animation/TweeningUtils/Source/TweeningUtilsEditor/Private/Widgets/MVC/TweenControllers.cpp
@@ -2,17 +2,59 @@
 
 #include "Widgets/MVC/TweenControllers.h"
 
-#include "Math/Models/FunctionBasedTweenModel.h"
+#include "TweeningToolsUserSettings.h"
+#include "Math/Abstraction/ITweenModelContainer.h"
 
 namespace UE::TweeningUtilsEditor
 {
+namespace Private
+{
+static void SetupUserPreferredFunctionConfig(
+	FTweenToolbarController& ToolbarController, const TSharedRef<ITweenModelContainer>& InTweenModels, FName UserPreferredFunctionContext
+	)
+{
+	const bool bShouldUsePreferredFunctions = UserPreferredFunctionContext != NAME_None;
+	if (!bShouldUsePreferredFunctions)
+	{
+		return;
+	}
+
+	// Restore currently selected function to the one that is in the config.
+	const FString* PreferredFunctionId = UTweeningToolsUserSettings::Get()->GetPreferredTweenFunction(UserPreferredFunctionContext);
+	FTweenModel* PreferredFunction = PreferredFunctionId ? InTweenModels->FindModelByIdentifier(*PreferredFunctionId) : nullptr;
+	if (PreferredFunction)
+	{
+		ToolbarController.SetSelectedTweenModel(*PreferredFunction);
+	}
+
+	// Save the key function to config when the user changes it.
+	ToolbarController.OnTweenFunctionChanged().AddLambda([InTweenModels, UserPreferredFunctionContext](const FTweenModel& TweenModel)
+	{
+		UTweeningToolsUserSettings::Get()->SetPreferredTweenFunction(
+			UserPreferredFunctionContext, InTweenModels->GetModelIdentifier(TweenModel)
+			);
+	});
+}
+}
+	
 FTweenControllers::FTweenControllers(
-	const TSharedRef<FFunctionBasedTweenModel>& InTweenModel,
 	const TSharedRef<FUICommandList>& InCommandList,
-	const TSharedRef<ITweenFunctionContainer>& InTweenFunctions
-	)
-	: ToolbarController(InTweenModel, InCommandList, InTweenFunctions)
-	, CycleFunctionController(InTweenModel, InTweenFunctions, InCommandList)
-	, MouseSlidingController(InTweenModel, InCommandList)
-	{}
+	const TSharedRef<ITweenModelContainer>& InTweenModels,
+	FName UserPreferredFunctionContext
+)
+	: ToolbarController(InCommandList, InTweenModels)
+	, CycleFunctionController(
+		ToolbarController.MakeSelectedConstTweenModelAttr(),
+		 InTweenModels,
+		FCycleFunctionController::FHandleTweenChange::CreateLambda([this](FTweenModel& TweenModel)
+		{
+			ToolbarController.SetSelectedTweenModel(TweenModel);
+		}),
+		InCommandList
+		)
+	, MouseSlidingController(ToolbarController.MakeSelectedTweenModelAttr(), InCommandList)
+{
+	check(InTweenModels->NumModels() > 0);
+	Private::SetupUserPreferredFunctionConfig(ToolbarController, InTweenModels, UserPreferredFunctionContext);
+}
 }
diff --git a/Engine/Plugins/Animation/TweeningUtils/Source/TweeningUtilsEditor/Private/Widgets/MVC/TweenMouseSlidingController.cpp b/Engine/Plugins/Animation/TweeningUtils/Source/TweeningUtilsEditor/Private/Widgets/MVC/TweenMouseSlidingController.cpp
--- a/Engine/Plugins/Animation/TweeningUtils/Source/TweeningUtilsEditor/Private/Widgets/MVC/TweenMouseSlidingController.cpp
+++ b/Engine/Plugins/Animation/TweeningUtils/Source/TweeningUtilsEditor/Private/Widgets/MVC/TweenMouseSlidingController.cpp
@@ -19,7 +19,7 @@ static TAutoConsoleVariable<float> CVarOverrideTweenSliderWidth(
 	
 FTweenMouseSlidingController::FTweenMouseSlidingController(
 	TAttribute<float> InMaxSlideWidthAttr,
-	const TSharedRef<FTweenModel>& InModel,
+	TAttribute<FTweenModel*> InTweenModelAttr,
 	const TSharedRef<FUICommandList>& InCommandList,
 	TSharedPtr<FUICommandInfo> InDragSliderCommand
 	)
@@ -31,31 +31,36 @@ FTweenMouseSlidingController::FTweenMouseSlidingController(
 		}),
 		InCommandList, MoveTemp(InDragSliderCommand)
 		)
-	, TweenModel(InModel)
+	, TweenModelAttr(MoveTemp(InTweenModelAttr))
 {
-	OnStartSliding().AddSP(TweenModel, &FTweenModel::StartBlendOperation);
+	OnStartSliding().AddRaw(this, &FTweenMouseSlidingController::OnSliderStartMove);
 	OnStopSliding().AddRaw(this, &FTweenMouseSlidingController::OnSliderEndMove);
 	OnUpdateSliding().AddRaw(this, &FTweenMouseSlidingController::OnSliderMove);
 }
 
-FTweenMouseSlidingController::FTweenMouseSlidingController(const TSharedRef<FTweenModel>& InModel, const TSharedRef<FUICommandList>& InCommandList)
+FTweenMouseSlidingController::FTweenMouseSlidingController(TAttribute<FTweenModel*> InTweenModelAttr, const TSharedRef<FUICommandList>& InCommandList)
 	: FTweenMouseSlidingController(
 		FTweeningUtilsStyle::Get().GetWidgetStyle<FTweenSliderStyle>("TweenSlider").BarDimensions.X,
-		InModel,
+		MoveTemp(InTweenModelAttr),
 		InCommandList,
 		FTweeningUtilsCommands::Get().DragAnimSliderTool
 	)
 {}
 
+void FTweenMouseSlidingController::OnSliderStartMove()
+{
+	TweenModelAttr.Get()->StartBlendOperation();
+}
+
 void FTweenMouseSlidingController::OnSliderEndMove()
 {
 	CurrentSliderPosition.Reset();
-	TweenModel->StopBlendOperation();
+	TweenModelAttr.Get()->StopBlendOperation();
 }
 	
 void FTweenMouseSlidingController::OnSliderMove(float InValue)
 {
 	CurrentSliderPosition = InValue;
-	TweenModel->BlendValues(InValue);
+	TweenModelAttr.Get()->BlendValues(InValue);
 }
 }
diff --git a/Engine/Plugins/Animation/TweeningUtils/Source/TweeningUtilsEditor/Private/Widgets/MVC/TweenToolbarController.cpp b/Engine/Plugins/Animation/TweeningUtils/Source/TweeningUtilsEditor/Private/Widgets/MVC/TweenToolbarController.cpp
--- a/Engine/Plugins/Animation/TweeningUtils/Source/TweeningUtilsEditor/Private/Widgets/MVC/TweenToolbarController.cpp
+++ b/Engine/Plugins/Animation/TweeningUtils/Source/TweeningUtilsEditor/Private/Widgets/MVC/TweenToolbarController.cpp
@@ -4,13 +4,13 @@
 
 #include "Framework/Commands/UICommandList.h"
 #include "Framework/MultiBox/MultiBoxBuilder.h"
-#include "Math/Abstraction/ITweenFunctionContainer.h"
-#include "Math/Models/FunctionBasedTweenModel.h"
 #include "Math/Color.h"
 #include "Styling/SlateStyle.h"
 #include "Styling/SlateStyleMacros.h"
 #include "Styling/SlateStyleRegistry.h"
 #include "TweeningUtilsCommands.h"
+#include "Math/Abstraction/ITweenModelContainer.h"
+#include "Math/Models/TweenModel.h"
 #include "Widgets/MVC/STweenView.h"
 
 namespace UE::TweeningUtilsEditor
@@ -19,7 +19,7 @@ class FOvershootButtonStyleHack : public FSlateStyleSet
 {
 public:
 
-	explicit FOvershootButtonStyleHack(const ITweenFunctionContainer& InTweenFunctions)
+	explicit FOvershootButtonStyleHack(const ITweenModelContainer& InTweenFunctions)
 		: FSlateStyleSet(*FString::Printf(TEXT("TweeningUtils_OvershootOverride_%s"), *FGuid::NewGuid().ToString()))
 	{
 		const FString PluginContentDir = FPaths::EnginePluginsDir() / TEXT("Animation") / TEXT("TweeningUtils") / TEXT("Resources");
@@ -30,10 +30,10 @@ public:
 		const FVector2D Icon20x20(20.0f, 20.0f);
 		FSlateBrush* BaseBrush = new IMAGE_BRUSH_SVG("Icons/OvershootMode", Icon20x20);
 		Set(TEXT("OvershootMode"), BaseBrush);
-		InTweenFunctions.ForEachFunction([this, &InTweenFunctions, &BaseBrush](const ITweenFunction& Function) 
+		InTweenFunctions.ForEachModel([this, &InTweenFunctions, &BaseBrush](const FTweenModel& Function) 
 		{
 			FSlateBrush* Brush = new FSlateBrush(*BaseBrush);
-			Brush->TintColor = InTweenFunctions.GetColorForFunction(Function);
+			Brush->TintColor = InTweenFunctions.GetColorForModel(Function);
 			Set(*GetOvershootButtonStyleName(InTweenFunctions, Function), Brush);
 		});
 		
@@ -46,32 +46,31 @@ public:
 	}
 
 	static FString GetOvershootButtonStyleName(
-		const ITweenFunctionContainer& InTweenFunctions, const ITweenFunction& InFunction
+		const ITweenModelContainer& InTweenFunctions, const FTweenModel& InFunction
 		)
 	{
-		return InTweenFunctions.GetFunctionIdentifier(InFunction);
+		return InTweenFunctions.GetModelIdentifier(InFunction);
 	}
 };
 	
 FTweenToolbarController::FTweenToolbarController(
-	const TSharedRef<FFunctionBasedTweenModel>& InModel,
 	const TSharedRef<FUICommandList>& InCommandList,
-	const TSharedRef<ITweenFunctionContainer>& InTweenFunctions
+	const TSharedRef<ITweenModelContainer>& InTweenFunctions,
+	int32 InInitialTweenModelIndex
 	)
-	: TweenModel(InModel)
-	, CommandList(InCommandList)
-	, TweenFunctions(InTweenFunctions)
+	: CommandList(InCommandList)
+	, TweenModels(InTweenFunctions)
+	, SelectedTweenModel(TweenModels->GetModel(InInitialTweenModelIndex))
 	, OverrideStyle(MakePimpl<FOvershootButtonStyleHack>(*InTweenFunctions))
 {
 	// Function changing commands
-	TweenFunctions->ForEachFunction([this](const ITweenFunction& InFunction)
+	TweenModels->ForEachModel([this](FTweenModel& InTweenModel)
 	{
-		CommandList->MapAction(TweenFunctions->GetCommandForFunction(
-			InFunction),
+		CommandList->MapAction(TweenModels->GetCommandForModel(InTweenModel),
 			FUIAction(
-			   FExecuteAction::CreateRaw(this, &FTweenToolbarController::SetBlendFunction, &InFunction),
+			   FExecuteAction::CreateRaw(this, &FTweenToolbarController::SetTweenModel, &InTweenModel),
 			   {},
-			   FIsActionChecked::CreateRaw(this, &FTweenToolbarController::IsBlendFunctionSelected, &InFunction)
+			   FIsActionChecked::CreateRaw(this, &FTweenToolbarController::IsTweenModelSelected, &InTweenModel)
 			)
 		);
 	});
@@ -88,9 +87,9 @@ FTweenToolbarController::FTweenToolbarController(
 FTweenToolbarController::~FTweenToolbarController()
 {
 	// Function changing commands
-	TweenFunctions->ForEachFunction([this](const ITweenFunction& InFunction)
+	TweenModels->ForEachModel([this](const FTweenModel& InTweenModel)
 	{
-		CommandList->UnmapAction(TweenFunctions->GetCommandForFunction(InFunction));
+		CommandList->UnmapAction(TweenModels->GetCommandForModel(InTweenModel));
 	});
 	
 	CommandList->UnmapAction(FTweeningUtilsCommands::Get().ToggleOvershootMode);
@@ -109,7 +108,8 @@ void FTweenToolbarController::AddToToolbar(FToolBarBuilder& ToolBarBuilder, FMak
 	ToolBarBuilder.SetLabelVisibility(EVisibility::Collapsed);
 		
 	ToolBarBuilder.AddWidget(
-		SNew(STweenView, TweenModel)
+		SNew(STweenView)
+		.TweenModel(MakeSelectedTweenModelAttr())
 		.SliderIcon(TAttribute<const FSlateBrush*>::CreateRaw(this, &FTweenToolbarController::GetIconForSlider))
 		.SliderColor(TAttribute<FLinearColor>::CreateRaw(this, &FTweenToolbarController::GetColorForSlider))
 		.OverrideSliderPosition(MoveTemp(InArgs.OverrideSliderPositionAttr))
@@ -121,13 +121,21 @@ void FTweenToolbarController::AddToToolbar(FToolBarBuilder& ToolBarBuilder, FMak
 		);
 }
 
+void FTweenToolbarController::SetSelectedTweenModel(FTweenModel& InTweenModel)
+{
+	if (ensure(TweenModels->Contains(InTweenModel)))
+	{
+		SetTweenModel(&InTweenModel);
+	}
+}
+
 TSharedRef<SWidget> FTweenToolbarController::MakeTweenModeMenu() const
 {
 	FMenuBuilder MenuBuilder(true, CommandList);
 
-	TweenFunctions->ForEachFunction([this, &MenuBuilder](const ITweenFunction& InFunction)
+	TweenModels->ForEachModel([this, &MenuBuilder](const FTweenModel& InTweenModel)
 	{
-		MenuBuilder.AddMenuEntry(TweenFunctions->GetCommandForFunction(InFunction));
+		MenuBuilder.AddMenuEntry(TweenModels->GetCommandForModel(InTweenModel));
 	});
 
 	return MenuBuilder.MakeWidget();
@@ -135,57 +143,60 @@ TSharedRef<SWidget> FTweenToolbarController::MakeTweenModeMenu() const
 
 FText FTweenToolbarController::GetLabelForComboBox() const
 {
-	return TweenFunctions->GetLabelForFunction(TweenModel->GetTweenFunction());
+	return TweenModels->GetLabelForModel(*SelectedTweenModel);
 }
 
 FText FTweenToolbarController::GetToolTipForComboBox() const
 {
-	return TweenFunctions->GetToolTipForFunction(TweenModel->GetTweenFunction());
+	return TweenModels->GetToolTipForModel(*SelectedTweenModel);
 }
 
 FSlateIcon FTweenToolbarController::GetIconForComboBox() const
 {
-	return TweenFunctions->GetCommandForFunction(TweenModel->GetTweenFunction())->GetIcon();
+	return TweenModels->GetCommandForModel(*SelectedTweenModel)->GetIcon();
 }
 
-void FTweenToolbarController::SetBlendFunction(const ITweenFunction* InFunction) const
+void FTweenToolbarController::SetTweenModel(FTweenModel* InTweenModel)
 {
-	TweenModel->SetTweenFunction(*InFunction);
-	OnTweenFunctionChangedDelegate.Broadcast(*InFunction);
+	SelectedTweenModel = InTweenModel;
+	OnTweenFunctionChangedDelegate.Broadcast(*SelectedTweenModel);
 }
 
-bool FTweenToolbarController::IsBlendFunctionSelected(const ITweenFunction* InFunction) const
+bool FTweenToolbarController::IsTweenModelSelected(FTweenModel* InTweenModel) const
 {
-	return &TweenModel->GetTweenFunction() == InFunction;
+	return SelectedTweenModel == InTweenModel;
 }
 
 void FTweenToolbarController::ToggleOvershootMode() const
 {
-	TweenModel->SetScaleMode(
-	   TweenModel->GetScaleMode() == ETweenScaleMode::Normalized ? ETweenScaleMode::Overshoot : ETweenScaleMode::Normalized
-	   );
+	TweenModels->ForEachModel([](FTweenModel& InTweenModel)
+	{
+		InTweenModel.SetScaleMode(
+			InTweenModel.GetScaleMode() == ETweenScaleMode::Normalized ? ETweenScaleMode::Overshoot : ETweenScaleMode::Normalized
+			);
+	});
 }
 
 bool FTweenToolbarController::IsOvershootModeEnabled() const
 {
-	return TweenModel->GetScaleMode() == ETweenScaleMode::Overshoot;
+	return SelectedTweenModel->GetScaleMode() == ETweenScaleMode::Overshoot;
 }
 
 FSlateIcon FTweenToolbarController::GetOvershootModeIcon() const
 {
 	return FSlateIcon(
 		OverrideStyle->GetStyleSetName(),
-		IsOvershootModeEnabled() ? TEXT("OvershootMode") : *OverrideStyle->GetOvershootButtonStyleName(*TweenFunctions, TweenModel->GetTweenFunction())
+		IsOvershootModeEnabled() ? TEXT("OvershootMode") : *OverrideStyle->GetOvershootButtonStyleName(*TweenModels, *SelectedTweenModel)
 		);
 }
 
 const FSlateBrush* FTweenToolbarController::GetIconForSlider() const
 {
-	return TweenFunctions->GetIconForFunction(TweenModel->GetTweenFunction());
+	return TweenModels->GetIconForModel(*SelectedTweenModel);
 }
 
 FLinearColor FTweenToolbarController::GetColorForSlider() const
 {
-	return TweenFunctions->GetColorForFunction(TweenModel->GetTweenFunction());
+	return TweenModels->GetColorForModel(*SelectedTweenModel);
 }
 }
diff --git a/Engine/Plugins/Animation/TweeningUtils/Source/TweeningUtilsEditor/Public/Math/Abstraction/BuiltInTweenFunction.h b/Engine/Plugins/Animation/TweeningUtils/Source/TweeningUtilsEditor/Public/Math/Abstraction/BuiltInTweenFunction.h
deleted file mode 100644
diff --git a/Engine/Plugins/Animation/TweeningUtils/Source/TweeningUtilsEditor/Public/Math/Abstraction/BuiltInTweenFunctionContainer.h b/Engine/Plugins/Animation/TweeningUtils/Source/TweeningUtilsEditor/Public/Math/Abstraction/BuiltInTweenFunctionContainer.h
deleted file mode 100644
diff --git a/Engine/Plugins/Animation/TweeningUtils/Source/TweeningUtilsEditor/Public/Math/Abstraction/BuiltInTweenModelContainer.h b/Engine/Plugins/Animation/TweeningUtils/Source/TweeningUtilsEditor/Public/Math/Abstraction/BuiltInTweenModelContainer.h
new file mode 100644
--- /dev/null
+++ b/Engine/Plugins/Animation/TweeningUtils/Source/TweeningUtilsEditor/Public/Math/Abstraction/BuiltInTweenModelContainer.h
@@ -0,0 +1,65 @@
+﻿// Copyright Epic Games, Inc. All Rights Reserved.
+
+#pragma once
+
+#include "ITweenModelContainer.h"
+#include "KeyBlendingAbstraction.h"
+#include "Math/Models/TweenModel.h"
+#include "Templates/UniquePtr.h"
+
+#include <array>
+#include <type_traits>
+
+namespace UE::TweeningUtilsEditor
+{
+/** Contains all the functions from this module, i.e. those specified by EBlendFunction.  */
+class TWEENINGUTILSEDITOR_API FBuiltInTweenModelContainer : public ITweenModelContainer
+{
+public:
+	
+	using FTweenArray = std::array<TUniquePtr<FTweenModel>, static_cast<int32>(EBlendFunction::Num)>;
+
+	explicit FBuiltInTweenModelContainer(FTweenArray InTweenModels) : Functions(MoveTemp(InTweenModels)) {}
+	/** Util constructs the tween models using a factory function. */
+	template<typename TConstructTweenModel> requires std::is_invocable_r_v<TUniquePtr<FTweenModel>, TConstructTweenModel, EBlendFunction>
+	explicit FBuiltInTweenModelContainer(TConstructTweenModel&& InMakeTweenModel);
+
+	/** @return Gets the type of blend function the tween function corresponds to, or unset if you pass in a function not owned by this container. */
+	TOptional<EBlendFunction> GetBlendFunctionType(const FTweenModel& InTweenModelToFind) const;
+
+	/** @return The blend function identified by the enum. */
+	const FTweenModel* GetBlendFunction(EBlendFunction InBlendFunction) const;
+
+	//~ Begin FTweenModelContainer Interface
+	virtual void ForEachModel(const TFunctionRef<void(FTweenModel&)> InConsumer) const override;
+	virtual FTweenModel* GetModel(int32 InIndex) const override;
+	virtual int32 NumModels() const override { return Functions.size(); }
+	virtual TSharedPtr<FUICommandInfo> GetCommandForModel(const FTweenModel& InTweenModel) const override;
+	virtual const FSlateBrush* GetIconForModel(const FTweenModel& InTweenModel) const override;
+	virtual FLinearColor GetColorForModel(const FTweenModel& InTweenModel) const override;
+	virtual FText GetLabelForModel(const FTweenModel& InTweenModel) const override;
+	virtual FText GetToolTipForModel(const FTweenModel& InTweenModel) const override;
+	virtual FString GetModelIdentifier(const FTweenModel& InTweenModel) const override;
+	//~ Begin FTweenModelContainer Interface
+
+private:
+
+	/** The built-in functions. */
+	const FTweenArray Functions;
+};
+
+template <typename TConstructTweenModel> requires std::is_invocable_r_v<TUniquePtr<FTweenModel>, TConstructTweenModel, EBlendFunction>
+FBuiltInTweenModelContainer::FBuiltInTweenModelContainer(TConstructTweenModel&& InMakeTweenModel)
+	: FBuiltInTweenModelContainer([&InMakeTweenModel]
+	{
+		FTweenArray Array;
+		int32 Index = 0;
+		ForEachBlendFunction([&InMakeTweenModel, &Array, &Index](EBlendFunction BlendFunction)
+		{
+			Array[Index] = InMakeTweenModel(BlendFunction);
+			++Index;
+		});
+		return Array;
+	}())
+{}
+}
diff --git a/Engine/Plugins/Animation/TweeningUtils/Source/TweeningUtilsEditor/Public/Math/Abstraction/ITweenFunction.h b/Engine/Plugins/Animation/TweeningUtils/Source/TweeningUtilsEditor/Public/Math/Abstraction/ITweenFunction.h
deleted file mode 100644
diff --git a/Engine/Plugins/Animation/TweeningUtils/Source/TweeningUtilsEditor/Public/Math/Abstraction/ITweenFunctionContainer.h b/Engine/Plugins/Animation/TweeningUtils/Source/TweeningUtilsEditor/Public/Math/Abstraction/ITweenFunctionContainer.h
deleted file mode 100644
diff --git a/Engine/Plugins/Animation/TweeningUtils/Source/TweeningUtilsEditor/Public/Math/Abstraction/ITweenModelContainer.h b/Engine/Plugins/Animation/TweeningUtils/Source/TweeningUtilsEditor/Public/Math/Abstraction/ITweenModelContainer.h
new file mode 100644
--- /dev/null
+++ b/Engine/Plugins/Animation/TweeningUtils/Source/TweeningUtilsEditor/Public/Math/Abstraction/ITweenModelContainer.h
@@ -0,0 +1,61 @@
+﻿// Copyright Epic Games, Inc. All Rights Reserved.
+
+#pragma once
+
+#include "HAL/Platform.h"
+#include "Templates/FunctionFwd.h"
+#include "Templates/SharedPointerFwd.h"
+
+class FString;
+class FText;
+class FUICommandInfo;
+struct FLinearColor;
+struct FSlateBrush;
+
+namespace UE::TweeningUtilsEditor
+{
+class FTweenModel;
+
+/**
+ * This acts as the source of FTweenModels displayed in UI, e.g. by FTweenToolbarController.
+ * The implementation is expected to be a static list of tween functions, i.e. each invocation of ForEachFunction is supposed to return the same functions.
+ */
+class ITweenModelContainer
+{
+public:
+
+	/** Iterates through all the functions */
+	virtual void ForEachModel(const TFunctionRef<void(FTweenModel&)> InConsumer) const = 0;
+
+	/** @return The function at InIndex. */
+	virtual FTweenModel* GetModel(int32 InIndex) const = 0;
+	/** @return The number of functions contained. */
+	virtual int32 NumModels() const = 0;
+	/** @return The index of InTweenModel if found and INDEX_NONE if not. */
+	TWEENINGUTILSEDITOR_API int32 IndexOf(const FTweenModel& InTweenModel) const;
+	/** @return Whether InIndex is a valid argument for GetFunction. */
+	bool IsValidIndex(int32 InIndex) const { return InIndex >= 0 && InIndex < NumModels(); }
+	/** @return Whether InTweenModel is contained by this model. */
+	bool Contains(const FTweenModel& InTweenModel) const { return IsValidIndex(IndexOf(InTweenModel)); }
+
+	/** @return The command that is used to select the function */
+	virtual TSharedPtr<FUICommandInfo> GetCommandForModel(const FTweenModel& InTweenModel) const = 0;
+
+	/** @return The un-tinted icon to display the function in the UI with */
+	virtual const FSlateBrush* GetIconForModel(const FTweenModel& InTweenModel) const = 0;
+	/** @return The color that represents the function in the UI */
+	virtual FLinearColor GetColorForModel(const FTweenModel& InTweenModel) const = 0;
+
+	/** @return The label to display in the combo button */
+	virtual FText GetLabelForModel(const FTweenModel& InTweenModel) const = 0;
+	/** @return The description to displayed, e.g. in tool tips. */
+	virtual FText GetToolTipForModel(const FTweenModel& InTweenModel) const = 0;
+
+	/** @return Identifier that uniquely identifies this function. Used e.g. to encode the function type in a config file. */
+	virtual FString GetModelIdentifier(const FTweenModel& InTweenModel) const = 0;
+	/** @return The function identified by the InIdentifier (as returned by a call to GetFunctionIdentifier). */
+	TWEENINGUTILSEDITOR_API FTweenModel* FindModelByIdentifier(const FString& InIdentifier);
+
+	virtual ~ITweenModelContainer() = default;
+};
+}
diff --git a/Engine/Plugins/Animation/TweeningUtils/Source/TweeningUtilsEditor/Public/Math/Abstraction/TweenModelArray.h b/Engine/Plugins/Animation/TweeningUtils/Source/TweeningUtilsEditor/Public/Math/Abstraction/TweenModelArray.h
new file mode 100644
--- /dev/null
+++ b/Engine/Plugins/Animation/TweeningUtils/Source/TweeningUtilsEditor/Public/Math/Abstraction/TweenModelArray.h
@@ -0,0 +1,65 @@
+﻿// Copyright Epic Games, Inc. All Rights Reserved.
+
+#pragma once
+
+#include "ITweenModelContainer.h"
+#include "KeyBlendingAbstraction.h"
+#include "Math/Models/TweenModel.h"
+#include "Templates/UniquePtr.h"
+
+#include <array>
+#include <type_traits>
+
+namespace UE::TweeningUtilsEditor
+{
+/** Contains all the functions from this module, i.e. those specified by EBlendFunction.  */
+class TWEENINGUTILSEDITOR_API FTweenModelArray : public ITweenModelContainer
+{
+public:
+	
+	using FTweenArray = std::array<TUniquePtr<FTweenModel>, static_cast<int32>(EBlendFunction::Num)>;
+
+	explicit FTweenModelArray(FTweenArray InTweenModels) : Models(MoveTemp(InTweenModels)) {}
+	/** Util constructs the tween models using a factory function. */
+	template<typename TConstructTweenModel> requires std::is_invocable_r_v<TUniquePtr<FTweenModel>, TConstructTweenModel, EBlendFunction>
+	explicit FTweenModelArray(TConstructTweenModel&& InMakeTweenModel);
+
+	/** @return Gets the type of blend function the tween function corresponds to, or unset if you pass in a function not owned by this container. */
+	TOptional<EBlendFunction> GetBlendFunctionType(const FTweenModel& InTweenModelToFind) const;
+
+	/** @return The blend function identified by the enum. */
+	const FTweenModel* GetBlendFunction(EBlendFunction InBlendFunction) const;
+
+	//~ Begin FTweenModelContainer Interface
+	virtual void ForEachModel(const TFunctionRef<void(FTweenModel&)> InConsumer) const override;
+	virtual FTweenModel* GetModel(int32 InIndex) const override;
+	virtual int32 NumModels() const override { return Models.size(); }
+	virtual TSharedPtr<FUICommandInfo> GetCommandForModel(const FTweenModel& InTweenModel) const override;
+	virtual const FSlateBrush* GetIconForModel(const FTweenModel& InTweenModel) const override;
+	virtual FLinearColor GetColorForModel(const FTweenModel& InTweenModel) const override;
+	virtual FText GetLabelForModel(const FTweenModel& InTweenModel) const override;
+	virtual FText GetToolTipForModel(const FTweenModel& InTweenModel) const override;
+	virtual FString GetModelIdentifier(const FTweenModel& InTweenModel) const override;
+	//~ Begin FTweenModelContainer Interface
+
+private:
+
+	/** The built-in functions. */
+	const FTweenArray Models;
+};
+
+template <typename TConstructTweenModel> requires std::is_invocable_r_v<TUniquePtr<FTweenModel>, TConstructTweenModel, EBlendFunction>
+FTweenModelArray::FTweenModelArray(TConstructTweenModel&& InMakeTweenModel)
+	: FTweenModelArray([&InMakeTweenModel]
+	{
+		FTweenArray Array;
+		int32 Index = 0;
+		ForEachBlendFunction([&InMakeTweenModel, &Array, &Index](EBlendFunction BlendFunction)
+		{
+			Array[Index] = InMakeTweenModel(BlendFunction);
+			++Index;
+		});
+		return Array;
+	}())
+{}
+}
diff --git a/Engine/Plugins/Animation/TweeningUtils/Source/TweeningUtilsEditor/Public/Math/Models/CurveTweenModel.h b/Engine/Plugins/Animation/TweeningUtils/Source/TweeningUtilsEditor/Public/Math/Models/CurveTweenModel.h
--- a/Engine/Plugins/Animation/TweeningUtils/Source/TweeningUtilsEditor/Public/Math/Models/CurveTweenModel.h
+++ b/Engine/Plugins/Animation/TweeningUtils/Source/TweeningUtilsEditor/Public/Math/Models/CurveTweenModel.h
@@ -2,8 +2,9 @@
 
 #pragma once
 
-#include "FunctionBasedTweenModel.h"
+#include "TweenModel.h"
 #include "Math/ContiguousKeyMapping.h"
+#include "Math/Abstraction/KeyBlendingAbstraction.h"
 #include "Templates/SharedPointer.h"
 
 class FCurveEditor;
@@ -11,11 +12,11 @@ class FCurveEditor;
 namespace UE::TweeningUtilsEditor
 {
 /** Tween the curves of the curve editor's selection. */
-class TWEENINGUTILSEDITOR_API FCurveTweenModel : public FFunctionBasedTweenModel
+class TWEENINGUTILSEDITOR_API FCurveTweenModel : public FTweenModel
 {
 public:
 
-	explicit FCurveTweenModel(TWeakPtr<FCurveEditor> InWeakCurveEditor, const ITweenFunction& InTweenFunction);
+	explicit FCurveTweenModel(TWeakPtr<FCurveEditor> InWeakCurveEditor, EBlendFunction InBlendFunction);
 
 	/** @return Whether it makes sense to call StartBlendOperation, etc. */
 	bool HasAnythingToBlend() const;
@@ -30,6 +31,8 @@ protected:
 	
 	/** The curve editor on which to tween the curves. */
 	const TWeakPtr<FCurveEditor> WeakCurveEditor;
+	/** The blend function to use for blending keys. */
+	const EBlendFunction BlendFunction;
 
 	/** Created in StartBlendOperation and used for the entirety of the blend operation. */
 	FContiguousKeyMapping ContiguousKeySelection;
diff --git a/Engine/Plugins/Animation/TweeningUtils/Source/TweeningUtilsEditor/Public/Math/Models/EditorTweenModel.h b/Engine/Plugins/Animation/TweeningUtils/Source/TweeningUtilsEditor/Public/Math/Models/EditorTweenModel.h
--- a/Engine/Plugins/Animation/TweeningUtils/Source/TweeningUtilsEditor/Public/Math/Models/EditorTweenModel.h
+++ b/Engine/Plugins/Animation/TweeningUtils/Source/TweeningUtilsEditor/Public/Math/Models/EditorTweenModel.h
@@ -3,8 +3,8 @@
 #pragma once
 
 #include "TangentFlatteningTweenProxy.h"
+#include "TweenModel.h"
 #include "TransactionalTweenModelProxy.h"
-#include "UserPreferredTweenProxy.h"
 
 #include <type_traits>
 
@@ -13,19 +13,16 @@ namespace UE::TweeningUtilsEditor
 /**
  * Extends TUnderlyingModel with commonly expected shared functionality:
  *	- Undo / redo
- *	- Keeps track of the last selected function and saves it in UTweeningToolsUserSettings.
  *	- As the tween function squishes the keys vertically, user-specified tangents are flattened.
- *	- TODO UE-234007: Synchronize tween function selection from different editor tools across all instances of the tweener.
- *	- In the future this can be extended to remove vertically stacked keys (same keys per frame). For now, no implemented tween function stacks keys.
+ *	- In the future this can be extended to remove vertically stacked keys (same keys per frame). For now, no implemented tween function stack keys.
  *
  * Construct like this:
  *  using FModel = TEditorTweenModel<FFooTweenModel> // FFooTweenModel should inherit from FTweenModel.
  *	const TSharedRef<FModel> TweenModel = MakeShared<FModel>(
- *		TEXT("YourFeatureName"), FunctionContainer	// TUserPreferredTweenProxy
  *		WeakCurveEditor,							// TTangentFlatteningTweenProxy arg
  *		Bar											// Any arguments that FFooTweenModel requires for construction
  *	);
  */
-template<typename TUnderlyingModel> requires std::is_base_of_v<FFunctionBasedTweenModel, TUnderlyingModel> // required by TUserPreferredTweenProxy.
-using TEditorTweenModel = TTransactionalTweenModelProxy<TUserPreferredTweenProxy<TTangentFlatteningTweenProxy<TUnderlyingModel>>>;
+template<typename TUnderlyingModel> requires std::is_base_of_v<FTweenModel, TUnderlyingModel>
+using TEditorTweenModel = TTransactionalTweenModelProxy<TTangentFlatteningTweenProxy<TUnderlyingModel>>;
 }
diff --git a/Engine/Plugins/Animation/TweeningUtils/Source/TweeningUtilsEditor/Public/Math/Models/FunctionBasedTweenModel.h b/Engine/Plugins/Animation/TweeningUtils/Source/TweeningUtilsEditor/Public/Math/Models/FunctionBasedTweenModel.h
deleted file mode 100644
diff --git a/Engine/Plugins/Animation/TweeningUtils/Source/TweeningUtilsEditor/Public/Math/Models/TweenModel.h b/Engine/Plugins/Animation/TweeningUtils/Source/TweeningUtilsEditor/Public/Math/Models/TweenModel.h
--- a/Engine/Plugins/Animation/TweeningUtils/Source/TweeningUtilsEditor/Public/Math/Models/TweenModel.h
+++ b/Engine/Plugins/Animation/TweeningUtils/Source/TweeningUtilsEditor/Public/Math/Models/TweenModel.h
@@ -2,7 +2,6 @@
 
 #pragma once
 
-#include "Misc/Optional.h"
 #include "Widgets/ETweenScaleMode.h"
 
 namespace UE::TweeningUtilsEditor
diff --git a/Engine/Plugins/Animation/TweeningUtils/Source/TweeningUtilsEditor/Public/Math/Models/UserPreferredTweenProxy.h b/Engine/Plugins/Animation/TweeningUtils/Source/TweeningUtilsEditor/Public/Math/Models/UserPreferredTweenProxy.h
deleted file mode 100644
diff --git a/Engine/Plugins/Animation/TweeningUtils/Source/TweeningUtilsEditor/Public/Widgets/MVC/CycleFunctionController.h b/Engine/Plugins/Animation/TweeningUtils/Source/TweeningUtilsEditor/Public/Widgets/MVC/CycleFunctionController.h
--- a/Engine/Plugins/Animation/TweeningUtils/Source/TweeningUtilsEditor/Public/Widgets/MVC/CycleFunctionController.h
+++ b/Engine/Plugins/Animation/TweeningUtils/Source/TweeningUtilsEditor/Public/Widgets/MVC/CycleFunctionController.h
@@ -2,6 +2,8 @@
 
 #pragma once
 
+#include "Delegates/Delegate.h"
+#include "Misc/Attribute.h"
 #include "Templates/SharedPointer.h"
 #include "Templates/UnrealTemplate.h"
 
@@ -10,41 +12,47 @@ class FUICommandList;
 
 namespace UE::TweeningUtilsEditor
 {
-class FFunctionBasedTweenModel;
 class FTweenModel;
-class ITweenFunctionContainer;
+class ITweenModelContainer;
 
 /** Cycles through functions on a tween model. */
 class TWEENINGUTILSEDITOR_API FCycleFunctionController : public FNoncopyable
 {
 public:
+
+	DECLARE_DELEGATE_OneParam(FHandleTweenChange, FTweenModel&);
 	
 	/**
-	 * @param InTweenModel The model on which the functions will be changed.
-	 * @param InTweenFunctions Determines the functions through which you can cycle
+	 * @param InCurrentTweenModelAttr The tween model that is currently selected.
+	 * @param InTweenModelContainer Determines the functions through which you can cycle
+	 * @param InHandleTweenChange Invoked to change the current tween function
 	 * @param InCommandList The command list to add / remove the command to / from.
 	 * @param InCycleCommand The command that cycles the function
 	 */
 	explicit FCycleFunctionController(
-		const TSharedRef<FFunctionBasedTweenModel>& InTweenModel,
-		const TSharedRef<ITweenFunctionContainer>& InTweenFunctions,
+		TAttribute<const FTweenModel*> InCurrentTweenModelAttr,
+		const TSharedRef<ITweenModelContainer>& InTweenModelContainer,
+		FHandleTweenChange InHandleTweenChange,
 		const TSharedRef<FUICommandList>& InCommandList,
 		TSharedPtr<FUICommandInfo> InCycleCommand
 		);
 	/** Version that defauls the cycle command to FTweenUtilsCommands::ChangeAnimSliderTool. */
 	explicit FCycleFunctionController(
-		const TSharedRef<FFunctionBasedTweenModel>& InTweenModel,
-		const TSharedRef<ITweenFunctionContainer>& InTweenFunctions,
+		TAttribute<const FTweenModel*> InCurrentTweenModelAttr,
+		const TSharedRef<ITweenModelContainer>& InTweenModelContainer,
+		FHandleTweenChange InHandleTweenChange,
 		const TSharedRef<FUICommandList>& InCommandList
 		); 
 	virtual ~FCycleFunctionController();
 	
 private:
 	
-	/** The model whose function is updated. */
-	const TSharedRef<FFunctionBasedTweenModel> TweenModel;
+	/** The tween model that is currently selected */
+	const TAttribute<const FTweenModel*> CurrentTweenModelAttr;
 	/** Holds the functions that can be cycled through. */
-	const TSharedRef<ITweenFunctionContainer> TweenFunctions;
+	const TSharedRef<ITweenModelContainer> TweenModelContainer;
+	/** Invoked to change the current tween function */
+	const FHandleTweenChange HandleTweenChangeDelegate;
 	
 	/** Used to bind and unbind the CycleCommand command. */
 	const TSharedRef<FUICommandList> CommandList;
diff --git a/Engine/Plugins/Animation/TweeningUtils/Source/TweeningUtilsEditor/Public/Widgets/MVC/STweenView.h b/Engine/Plugins/Animation/TweeningUtils/Source/TweeningUtilsEditor/Public/Widgets/MVC/STweenView.h
--- a/Engine/Plugins/Animation/TweeningUtils/Source/TweeningUtilsEditor/Public/Widgets/MVC/STweenView.h
+++ b/Engine/Plugins/Animation/TweeningUtils/Source/TweeningUtilsEditor/Public/Widgets/MVC/STweenView.h
@@ -21,6 +21,9 @@ class TWEENINGUTILSEDITOR_API STweenView : public SCompoundWidget
 public:
 
 	SLATE_BEGIN_ARGS(STweenView) {}
+		/** Gets the tween model this view is driving. */
+		SLATE_ATTRIBUTE(FTweenModel*, TweenModel)
+		
 		/** The icon to place in the slider button. */
 		SLATE_ATTRIBUTE(const FSlateBrush*, SliderIcon)
 		/** The main color. It tints the slider button and the points. */
@@ -29,19 +32,19 @@ public:
 		SLATE_ATTRIBUTE(TOptional<float>, OverrideSliderPosition)
 	SLATE_END_ARGS()
 
-	void Construct(const FArguments& InArgs, const TSharedRef<FTweenModel>& InTweenModel);
+	void Construct(const FArguments& InArgs);
 
 private:
 
 	/** The model that these controls are driving */
-	TSharedPtr<FTweenModel> TweenModel;
+	TAttribute<FTweenModel*> TweenModelAttr;
 	
 	ETweenScaleMode GetBarRenderMode() const;
 	
-	void OnDragStarted();
-	void OnDragEnded();
-	void OnDragValueUpdated(float Value);
-	void OnPointPicked(float Value);
-	float MapSliderValueToBlendValue(float Value);
+	void OnDragStarted() const;
+	void OnDragEnded() const;
+	void OnDragValueUpdated(float Value) const;
+	void OnPointPicked(float Value) const;
+	float MapSliderValueToBlendValue(float Value) const;
 };
 }
diff --git a/Engine/Plugins/Animation/TweeningUtils/Source/TweeningUtilsEditor/Public/Widgets/MVC/TweenControllers.h b/Engine/Plugins/Animation/TweeningUtils/Source/TweeningUtilsEditor/Public/Widgets/MVC/TweenControllers.h
--- a/Engine/Plugins/Animation/TweeningUtils/Source/TweeningUtilsEditor/Public/Widgets/MVC/TweenControllers.h
+++ b/Engine/Plugins/Animation/TweeningUtils/Source/TweeningUtilsEditor/Public/Widgets/MVC/TweenControllers.h
@@ -11,9 +11,6 @@ class FUICommandList;
 
 namespace UE::TweeningUtilsEditor
 {
-class FFunctionBasedTweenModel;
-class ITweenFunctionContainer;
-
 /** Holds functionality that may be reused across modules to add tweening widgets to a toolbar. */
 struct TWEENINGUTILSEDITOR_API FTweenControllers
 {
@@ -25,11 +22,16 @@ struct TWEENINGUTILSEDITOR_API FTweenControllers
 
 	/** Allows indirect movement of the slider by using u + Move Mouse */
 	FTweenMouseSlidingController MouseSlidingController;
-	
+
+	/**
+	 * @param InCommandList The command list to bind commands to
+	 * @param InTweenModels Contains all tween models that can be selected by the tweening widgets
+	 * @param UserPreferredFunctionContext Optional. The config key under which the selected tween function is saved. The function currently saved under this key will be the first one to be selected.
+	 */
 	explicit FTweenControllers(
-		const TSharedRef<FFunctionBasedTweenModel>& InTweenModel,
 		const TSharedRef<FUICommandList>& InCommandList,
-		const TSharedRef<ITweenFunctionContainer>& InTweenFunctions
+		const TSharedRef<ITweenModelContainer>& InTweenModels,
+		FName UserPreferredFunctionContext = NAME_None
 		);
 };
 }
diff --git a/Engine/Plugins/Animation/TweeningUtils/Source/TweeningUtilsEditor/Public/Widgets/MVC/TweenMouseSlidingController.h b/Engine/Plugins/Animation/TweeningUtils/Source/TweeningUtilsEditor/Public/Widgets/MVC/TweenMouseSlidingController.h
--- a/Engine/Plugins/Animation/TweeningUtils/Source/TweeningUtilsEditor/Public/Widgets/MVC/TweenMouseSlidingController.h
+++ b/Engine/Plugins/Animation/TweeningUtils/Source/TweeningUtilsEditor/Public/Widgets/MVC/TweenMouseSlidingController.h
@@ -21,13 +21,13 @@ public:
 	/**
 	 * Version that allows you to specify every detail
 	 * @param InMaxSlideWidthAttr The total size of the "invisible" slider. The mouse can be moved left and right InMaxSlideWidthAttr.Get / 2.0 before being clamped.
-	 * @param InModel The tween model that the sliding will update.
+	 * @param InTweenModelAttr The tween model that the sliding will update.
 	 * @param InCommandList The command list to bind / unbind the command with.
 	 * @param InDragSliderCommand The command that triggers detection of mouse movement.
 	 */
 	explicit FTweenMouseSlidingController(
 		TAttribute<float> InMaxSlideWidthAttr,
-		const TSharedRef<FTweenModel>& InModel,
+		TAttribute<FTweenModel*> InTweenModelAttr,
 		const TSharedRef<FUICommandList>& InCommandList,
 		TSharedPtr<FUICommandInfo> InDragSliderCommand
 		);
@@ -36,7 +36,7 @@ public:
 	 * The slider width is based on the size of the slider in the toolbar(= value from FTweeningUtilsStyle).
 	 */
 	explicit FTweenMouseSlidingController(
-		const TSharedRef<FTweenModel>& InModel,
+		TAttribute<FTweenModel*> InTweenModelAttr,
 		const TSharedRef<FUICommandList>& InCommandList
 	);
 
@@ -46,11 +46,12 @@ public:
 private:
 	
 	/** Does the actual blending. */
-	const TSharedRef<FTweenModel> TweenModel;
+	const TAttribute<FTweenModel*> TweenModelAttr;
 
 	/** The current slider position that is being driven by the mouse. Range [-1,1]. */
 	TOptional<float> CurrentSliderPosition;
 
+	void OnSliderStartMove();
 	void OnSliderEndMove();
 	void OnSliderMove(float InValue);
 };
diff --git a/Engine/Plugins/Animation/TweeningUtils/Source/TweeningUtilsEditor/Public/Widgets/MVC/TweenToolbarController.h b/Engine/Plugins/Animation/TweeningUtils/Source/TweeningUtilsEditor/Public/Widgets/MVC/TweenToolbarController.h
--- a/Engine/Plugins/Animation/TweeningUtils/Source/TweeningUtilsEditor/Public/Widgets/MVC/TweenToolbarController.h
+++ b/Engine/Plugins/Animation/TweeningUtils/Source/TweeningUtilsEditor/Public/Widgets/MVC/TweenToolbarController.h
@@ -3,6 +3,7 @@
 #pragma once
 
 #include "Misc/Attribute.h"
+#include "Misc/NotNull.h"
 #include "Templates/PimplPtr.h"
 #include "Templates/SharedPointer.h"
 #include "Templates/UnrealTemplate.h"
@@ -18,9 +19,8 @@ template<typename OptionalType> struct TOptional;
 
 namespace UE::TweeningUtilsEditor
 {
-class FFunctionBasedTweenModel;
-class ITweenFunction;
-class ITweenFunctionContainer;
+class ITweenModelContainer;
+class FTweenModel;
 	
 /**
  * Binds together a STweenView and FTweenModel in the context of adding the controls to a toolbar widget.
@@ -36,14 +36,14 @@ class TWEENINGUTILSEDITOR_API FTweenToolbarController : public FNoncopyable
 public:
 	
 	/**
-	 * @param InModel The model that the UI will be driving.
 	 * @param InCommandList Used to bind the commands added to the toolbar.
 	 * @param InTweenFunctions The functions that can be selected
+	 * @param InInitialTweenModelIndex Index of the tween model to initially be selected.
 	 */
 	explicit FTweenToolbarController(
-		const TSharedRef<FFunctionBasedTweenModel>& InModel,
 		const TSharedRef<FUICommandList>& InCommandList,
-		const TSharedRef<ITweenFunctionContainer>& InTweenFunctions
+		const TSharedRef<ITweenModelContainer>& InTweenFunctions,
+		int32 InInitialTweenModelIndex = 0
 		);
 	~FTweenToolbarController();
 
@@ -63,20 +63,28 @@ public:
 	 */
 	void AddToToolbar(FToolBarBuilder& ToolbarBuilder, FMakeWidgetArgs InArgs = {}) const;
 
-	DECLARE_MULTICAST_DELEGATE_OneParam(FOnTweenFunctionChanged, const ITweenFunction&);
+	/** Sets the currently selected tween model. Asserts that the tween model is contained in TweenModels. */
+	void SetSelectedTweenModel(FTweenModel& InTweenModel);
+
+	/** @return Gets the currently selected tween model */
+	FTweenModel& GetSelectedTweenModel() const { return *SelectedTweenModel; }
+	TAttribute<FTweenModel*> MakeSelectedTweenModelAttr() const { return TAttribute<FTweenModel*>::CreateLambda([this]{ return SelectedTweenModel; }); }
+	TAttribute<const FTweenModel*> MakeSelectedConstTweenModelAttr() const { return TAttribute<const FTweenModel*>::CreateLambda([this]{ return SelectedTweenModel; }); }
+
+	DECLARE_MULTICAST_DELEGATE_OneParam(FOnTweenFunctionChanged, const FTweenModel&);
 	/** Invoked when the tween function is changed by this controller. */
 	FOnTweenFunctionChanged& OnTweenFunctionChanged() { return OnTweenFunctionChangedDelegate; }
 
 private:
 
-	/** Receives the blend values */
-	const TSharedRef<FFunctionBasedTweenModel> TweenModel;
 	/** Used when the combo button is opened. */
 	const TSharedRef<FUICommandList> CommandList;
 	/** Holds the functions that can be selected and gives us information about how to display them. */
-	const TSharedRef<ITweenFunctionContainer> TweenFunctions;
+	const TSharedRef<ITweenModelContainer> TweenModels;
+	/** The selected tween model. Points into TweenModels. */
+	TNotNull<FTweenModel*> SelectedTweenModel;
 
-	/** Used to override the color of the overshoot button to match the color of the function. */
+	/** Used to override the color of the overshoot button to match the color of the tween model. */
 	const TPimplPtr<class FOvershootButtonStyleHack> OverrideStyle;
 
 	/** Invoked when the tween function is changed by this controller. */
@@ -88,8 +96,8 @@ private:
 	FText GetToolTipForComboBox() const;
 	FSlateIcon GetIconForComboBox() const;
 	
-	void SetBlendFunction(const ITweenFunction* InFunction) const;
-	bool IsBlendFunctionSelected(const ITweenFunction* InFunction) const;
+	void SetTweenModel(FTweenModel* InTweenModel);
+	bool IsTweenModelSelected(FTweenModel* InTweenModel) const;
 	
 	void ToggleOvershootMode() const;
 	bool IsOvershootModeEnabled() const;
diff --git a/Engine/Plugins/Editor/CurveEditorTools/Source/CurveEditorTools/Private/Extensions/Tweening/ContextAwareTweeningModel.cpp b/Engine/Plugins/Editor/CurveEditorTools/Source/CurveEditorTools/Private/Extensions/Tweening/ContextAwareTweeningModel.cpp
--- a/Engine/Plugins/Editor/CurveEditorTools/Source/CurveEditorTools/Private/Extensions/Tweening/ContextAwareTweeningModel.cpp
+++ b/Engine/Plugins/Editor/CurveEditorTools/Source/CurveEditorTools/Private/Extensions/Tweening/ContextAwareTweeningModel.cpp
@@ -6,14 +6,10 @@
 
 namespace UE::CurveEditorTools
 {
-FContextAwareTweeningModel::FContextAwareTweeningModel(
-	TWeakPtr<FCurveEditor> InWeakCurveEditor,
-	const TweeningUtilsEditor::ITweenFunction& InTweenFunction
-	)
-	: FFunctionBasedTweenModel(InTweenFunction)
-	, CurveEditor(InWeakCurveEditor)
-	, SelectionBasedModel(InWeakCurveEditor, InTweenFunction)
-	, NoKeySelectionModel(MakeShared<FKeyInsertingTweenModel>(InWeakCurveEditor, InTweenFunction))
+FContextAwareTweeningModel::FContextAwareTweeningModel(TWeakPtr<FCurveEditor> InWeakCurveEditor, TweeningUtilsEditor::EBlendFunction InBlendFunction)
+	: CurveEditor(InWeakCurveEditor)
+	, SelectionBasedModel(InWeakCurveEditor, InBlendFunction)
+	, NoKeySelectionModel(MakeShared<FKeyInsertingTweenModel>(InWeakCurveEditor, InBlendFunction))
 {}
 
 void FContextAwareTweeningModel::SetScaleMode(TweeningUtilsEditor::ETweenScaleMode InMode)
@@ -26,19 +22,9 @@ void FContextAwareTweeningModel::SetScaleMode(TweeningUtilsEditor::ETweenScaleMo
 	}
 }
 
-void FContextAwareTweeningModel::SetTweenFunction(const TweeningUtilsEditor::ITweenFunction& InTweenFunction)
-{
-	Super::SetTweenFunction(InTweenFunction);
-	SelectionBasedModel.SetTweenFunction(InTweenFunction);
-	if (NoKeySelectionModel)
-	{
-		NoKeySelectionModel->SetTweenFunction(InTweenFunction);
-	}
-}
-
 void FContextAwareTweeningModel::StartBlendOperation()
 {
-	AccessCurrentModel([this](FFunctionBasedTweenModel& Model)
+	AccessCurrentModel([this](FTweenModel& Model)
 	{
 		Model.StartBlendOperation();
 	});
@@ -46,7 +32,7 @@ void FContextAwareTweeningModel::StartBlendOperation()
 
 void FContextAwareTweeningModel::StopBlendOperation()
 {
-	AccessCurrentModel([](FFunctionBasedTweenModel& Model)
+	AccessCurrentModel([](FTweenModel& Model)
 	{
 		Model.StopBlendOperation();
 	});
@@ -54,13 +40,13 @@ void FContextAwareTweeningModel::StopBlendOperation()
 
 void FContextAwareTweeningModel::BlendValues(float InNormalizedValue)
 {
-	AccessCurrentModel([InNormalizedValue](FFunctionBasedTweenModel& Model)
+	AccessCurrentModel([InNormalizedValue](FTweenModel& Model)
 	{
 		Model.BlendValues(InNormalizedValue);
 	});
 }
 
-template <typename TCallback> requires std::is_invocable_v<TCallback, TweeningUtilsEditor::FFunctionBasedTweenModel&>
+template <typename TCallback> requires std::is_invocable_v<TCallback, TweeningUtilsEditor::FTweenModel&>
 void FContextAwareTweeningModel::AccessCurrentModel(TCallback&& InCallback)
 {
 	if (SelectionBasedModel.HasAnythingToBlend())
@@ -72,4 +58,4 @@ void FContextAwareTweeningModel::AccessCurrentModel(TCallback&& InCallback)
 		InCallback(*NoKeySelectionModel.Get());
 	}
 }
-}
+}
\ No newline at end of file
diff --git a/Engine/Plugins/Editor/CurveEditorTools/Source/CurveEditorTools/Private/Extensions/Tweening/ContextAwareTweeningModel.h b/Engine/Plugins/Editor/CurveEditorTools/Source/CurveEditorTools/Private/Extensions/Tweening/ContextAwareTweeningModel.h
--- a/Engine/Plugins/Editor/CurveEditorTools/Source/CurveEditorTools/Private/Extensions/Tweening/ContextAwareTweeningModel.h
+++ b/Engine/Plugins/Editor/CurveEditorTools/Source/CurveEditorTools/Private/Extensions/Tweening/ContextAwareTweeningModel.h
@@ -2,9 +2,7 @@
 
 #pragma once
 
-#include "Containers/ContainersFwd.h"
 #include "Math/Models/CurveTweenModel.h"
-#include "Math/Models/FunctionBasedTweenModel.h"
 
 #include <type_traits>
 
@@ -19,23 +17,19 @@ namespace UE::CurveEditorTools
  * - When no keys are selected, we'll insert keys at the time controller's scrub position to all visible curves (those selected in tree view).
  *   If there's no such controller, nothing happens.
  */
-class FContextAwareTweeningModel : public TweeningUtilsEditor::FFunctionBasedTweenModel
+class FContextAwareTweeningModel : public TweeningUtilsEditor::FTweenModel
 {
-	using Super = FFunctionBasedTweenModel;
+	using Super = FTweenModel;
 public:
 	
 	/**
 	 * @param InWeakCurveEditor Used to get key info
-	 * @param InInitialTweenFunction The initial function to blend with
+	 * @param InBlendFunction The blend function to blend with
 	 */
-	explicit FContextAwareTweeningModel(
-		TWeakPtr<FCurveEditor> InWeakCurveEditor,
-		const TweeningUtilsEditor::ITweenFunction& InInitialTweenFunction
-		);
+	explicit FContextAwareTweeningModel(TWeakPtr<FCurveEditor> InWeakCurveEditor, TweeningUtilsEditor::EBlendFunction InBlendFunction);
 
 	//~ Begin FTweenModel Interface
 	virtual void SetScaleMode(TweeningUtilsEditor::ETweenScaleMode InMode) override;
-	virtual void SetTweenFunction(const TweeningUtilsEditor::ITweenFunction& InTweenFunction) override;
 	virtual void StartBlendOperation() override;
 	virtual void StopBlendOperation() override;
 	virtual void BlendValues(float InNormalizedValue) override;
@@ -49,9 +43,9 @@ private:
 	/** Tweens when there is a key selection. */
 	TweeningUtilsEditor::FCurveTweenModel SelectionBasedModel;
 	/** Injected model that tweens when there is no key selection. Can be nullptr. */
-	const TSharedPtr<FFunctionBasedTweenModel> NoKeySelectionModel;
+	const TSharedPtr<FTweenModel> NoKeySelectionModel;
 
-	template<typename TCallback> requires std::is_invocable_v<TCallback, TweeningUtilsEditor::FFunctionBasedTweenModel&>
+	template<typename TCallback> requires std::is_invocable_v<TCallback, TweeningUtilsEditor::FTweenModel&>
 	void AccessCurrentModel(TCallback&& InCallback);
 };
 }
diff --git a/Engine/Plugins/Editor/CurveEditorTools/Source/CurveEditorTools/Private/Extensions/Tweening/KeyInsertingTweenModel.cpp b/Engine/Plugins/Editor/CurveEditorTools/Source/CurveEditorTools/Private/Extensions/Tweening/KeyInsertingTweenModel.cpp
--- a/Engine/Plugins/Editor/CurveEditorTools/Source/CurveEditorTools/Private/Extensions/Tweening/KeyInsertingTweenModel.cpp
+++ b/Engine/Plugins/Editor/CurveEditorTools/Source/CurveEditorTools/Private/Extensions/Tweening/KeyInsertingTweenModel.cpp
@@ -4,14 +4,13 @@
 
 #include "CurveEditor.h"
 #include "Math/CurveBlending.h"
-#include "Math/Abstraction/ITweenFunction.h"
 
 namespace UE::CurveEditorTools
 {
 FKeyInsertingTweenModel::FKeyInsertingTweenModel(
-	TWeakPtr<FCurveEditor> InWeakCurveEditor, const TweeningUtilsEditor::ITweenFunction& InInitialFunction
+	TWeakPtr<FCurveEditor> InWeakCurveEditor, TweeningUtilsEditor::EBlendFunction InBlendFunction
 	)
-	: FCurveTweenModel(MoveTemp(InWeakCurveEditor), InInitialFunction)
+	: FCurveTweenModel(MoveTemp(InWeakCurveEditor), InBlendFunction)
 {}
 
 void FKeyInsertingTweenModel::StartBlendOperation()
diff --git a/Engine/Plugins/Editor/CurveEditorTools/Source/CurveEditorTools/Private/Extensions/Tweening/KeyInsertingTweenModel.h b/Engine/Plugins/Editor/CurveEditorTools/Source/CurveEditorTools/Private/Extensions/Tweening/KeyInsertingTweenModel.h
--- a/Engine/Plugins/Editor/CurveEditorTools/Source/CurveEditorTools/Private/Extensions/Tweening/KeyInsertingTweenModel.h
+++ b/Engine/Plugins/Editor/CurveEditorTools/Source/CurveEditorTools/Private/Extensions/Tweening/KeyInsertingTweenModel.h
@@ -3,7 +3,6 @@
 #pragma once
 
 #include "Math/Models/CurveTweenModel.h"
-#include "Templates/SharedPointer.h"
 
 class FCurveEditor;
 
@@ -14,10 +13,7 @@ class FKeyInsertingTweenModel : public TweeningUtilsEditor::FCurveTweenModel
 {
 public:
 	
-	explicit FKeyInsertingTweenModel(
-		TWeakPtr<FCurveEditor> InWeakCurveEditor,
-		const TweeningUtilsEditor::ITweenFunction& InInitialFunction
-		);
+	explicit FKeyInsertingTweenModel(TWeakPtr<FCurveEditor> InWeakCurveEditor, TweeningUtilsEditor::EBlendFunction InBlendFunction);
 
 	//~ Begin FTweenModel
 	virtual void StartBlendOperation() override;
diff --git a/Engine/Plugins/Editor/CurveEditorTools/Source/CurveEditorTools/Private/Extensions/Tweening/TweenEditorExtension.cpp b/Engine/Plugins/Editor/CurveEditorTools/Source/CurveEditorTools/Private/Extensions/Tweening/TweenEditorExtension.cpp
--- a/Engine/Plugins/Editor/CurveEditorTools/Source/CurveEditorTools/Private/Extensions/Tweening/TweenEditorExtension.cpp
+++ b/Engine/Plugins/Editor/CurveEditorTools/Source/CurveEditorTools/Private/Extensions/Tweening/TweenEditorExtension.cpp
@@ -6,7 +6,7 @@
 #include "CurveEditor.h"
 #include "Framework/MultiBox/MultiBoxBuilder.h"
 #include "Framework/MultiBox/MultiBoxExtender.h"
-#include "Math/Abstraction/BuiltInTweenFunctionContainer.h"
+#include "Math/Abstraction/TweenModelArray.h"
 #include "Math/Models/EditorTweenModel.h"
 #include "Misc/Optional.h"
 #include "Widgets/MVC/TweenToolbarController.h"
@@ -15,7 +15,16 @@ namespace UE::CurveEditorTools
 {
 FTweenEditorExtension::FTweenEditorExtension(TWeakPtr<FCurveEditor> InCurveEditor)
 	: WeakCurveEditor(InCurveEditor)
-	, BlendFunctionContainer(MakeShared<TweeningUtilsEditor::FBuiltInTweenFunctionContainer>())
+	, TweenModelContainer(
+		MakeShared<TweeningUtilsEditor::FTweenModelArray>([this](TweeningUtilsEditor::EBlendFunction BlendFunction)
+		{
+			using FModel = TweeningUtilsEditor::TEditorTweenModel<FContextAwareTweeningModel>;
+			return MakeUnique<FModel>(
+				WeakCurveEditor,							// TTangentFlatteningTweenProxy arg
+				WeakCurveEditor, BlendFunction			// FContextAwareTweeningModel args
+				);
+		})
+		)
 {}
 
 void FTweenEditorExtension::BindCommands(TSharedRef<FUICommandList> InCommandList)
@@ -61,19 +70,7 @@ bool FTweenEditorExtension::InitControllers(TSharedRef<FUICommandList> InCommand
 
 	if (const TSharedPtr<FCurveEditor> CurveEditorPin = WeakCurveEditor.Pin(); ensure(CurveEditorPin))
 	{
-		using namespace TweeningUtilsEditor;
-		
-		const FName FeatureName = TEXT("CurveEditorTween"); // The key in UTweeningToolsUserSettings from which to look up the preferred function.
-		const ITweenFunction& DefaultFunction = *BlendFunctionContainer->GetFunction(0); // Use 1st function if no preferred function is in config.
-		
-		using FModel = TEditorTweenModel<FContextAwareTweeningModel>;
-		const TSharedRef<FModel> TweenModel = MakeShared<FModel>(
-			FeatureName, *BlendFunctionContainer,	// TUserPreferredTweenProxy
-			WeakCurveEditor,							// TTangentFlatteningTweenProxy arg
-			WeakCurveEditor, DefaultFunction			// FContextAwareTweeningModel args
-			);
-		TweenControllers.Emplace(TweenModel, InCommandList, BlendFunctionContainer);
-		
+		TweenControllers.Emplace(InCommandList, TweenModelContainer, TEXT("CurveEditorTween"));
 		return true;
 	}
 	return false;
diff --git a/Engine/Plugins/Editor/CurveEditorTools/Source/CurveEditorTools/Private/Extensions/Tweening/TweenEditorExtension.h b/Engine/Plugins/Editor/CurveEditorTools/Source/CurveEditorTools/Private/Extensions/Tweening/TweenEditorExtension.h
--- a/Engine/Plugins/Editor/CurveEditorTools/Source/CurveEditorTools/Private/Extensions/Tweening/TweenEditorExtension.h
+++ b/Engine/Plugins/Editor/CurveEditorTools/Source/CurveEditorTools/Private/Extensions/Tweening/TweenEditorExtension.h
@@ -9,10 +9,7 @@
 
 class FCurveEditor;
 class FExtender;
-namespace UE::TweeningUtilsEditor { class FBuiltInTweenFunctionContainer; }
-namespace UE::TweeningUtilsEditor { class FFunctionBasedTweenModel; }
-namespace UE::TweeningUtilsEditor { class ITweenFunction; }
-namespace UE::TweeningUtilsEditor { class ITweenFunctionContainer; }
+namespace UE::TweeningUtilsEditor { class FTweenModelArray; }
 
 namespace UE::CurveEditorTools
 {
@@ -34,7 +31,7 @@ private:
 	const TWeakPtr<FCurveEditor> WeakCurveEditor;
 
 	/** Holds the functions that ToolbarController shows: just the default, built-in ones.*/
-	const TSharedRef<TweeningUtilsEditor::FBuiltInTweenFunctionContainer> BlendFunctionContainer;
+	const TSharedRef<TweeningUtilsEditor::FTweenModelArray> TweenModelContainer;
 
 	/** Created by InitControllers. */
 	TOptional<TweeningUtilsEditor::FTweenControllers> TweenControllers;
