Tweening Utils: Combo button entries are now "Tween", etc. instead of "Set Tween", etc.

diff --git a/Engine/Plugins/Animation/TweeningUtils/Source/TweeningUtilsEditor/Private/TweeningUtilsCommands.cpp b/Engine/Plugins/Animation/TweeningUtils/Source/TweeningUtilsEditor/Private/TweeningUtilsCommands.cpp
--- a/Engine/Plugins/Animation/TweeningUtils/Source/TweeningUtilsEditor/Private/TweeningUtilsCommands.cpp
+++ b/Engine/Plugins/Animation/TweeningUtils/Source/TweeningUtilsEditor/Private/TweeningUtilsCommands.cpp
@@ -23,12 +23,12 @@ FTweeningUtilsCommands::FTweeningUtilsCommands()
 void FTweeningUtilsCommands::RegisterCommands()
 {
 	static_assert(static_cast<int32>(EBlendFunction::Num) == 6, "Add a command here");
-	UI_COMMAND(SetControlsToTween, "Set Tween ", "Sets the tween mode to Tween", EUserInterfaceActionType::RadioButton, FInputChord() );
-	UI_COMMAND(SetTweenPushPull, "Set Push / Pull", "Sets the tween mode to Push / Pull", EUserInterfaceActionType::RadioButton, FInputChord() );
-	UI_COMMAND(SetTweenBlendNeighbor, "Set Blend Neighbor", "Sets the tween mode to Blend Neighbour", EUserInterfaceActionType::RadioButton, FInputChord() );
-	UI_COMMAND(SetTweenBlendRelative, "Set Blend Relative", "Sets the tween mode to Blend Relative", EUserInterfaceActionType::RadioButton, FInputChord() );
-	UI_COMMAND(SetTweenBlendEase, "Set Blend Ease", "Sets the tween mode to Blend Ease", EUserInterfaceActionType::RadioButton, FInputChord() );
-	UI_COMMAND(SetTweenSmoothRough, "Set Smooth / Rough", "Sets the tween mode to Smooth / Rough", EUserInterfaceActionType::RadioButton, FInputChord() );
+	UI_COMMAND(SetControlsToTween, "Tween ", "Sets the tween mode to Tween", EUserInterfaceActionType::RadioButton, FInputChord() );
+	UI_COMMAND(SetTweenPushPull, "Push / Pull", "Sets the tween mode to Push / Pull", EUserInterfaceActionType::RadioButton, FInputChord() );
+	UI_COMMAND(SetTweenBlendNeighbor, "Blend Neighbor", "Sets the tween mode to Blend Neighbour", EUserInterfaceActionType::RadioButton, FInputChord() );
+	UI_COMMAND(SetTweenBlendRelative, "Blend Relative", "Sets the tween mode to Blend Relative", EUserInterfaceActionType::RadioButton, FInputChord() );
+	UI_COMMAND(SetTweenBlendEase, "Blend Ease", "Sets the tween mode to Blend Ease", EUserInterfaceActionType::RadioButton, FInputChord() );
+	UI_COMMAND(SetTweenSmoothRough, "Smooth / Rough", "Sets the tween mode to Smooth / Rough", EUserInterfaceActionType::RadioButton, FInputChord() );
 	UI_COMMAND(ToggleOvershootMode, "Overshoot mode", "Overshoot extends the available range over 100%", EUserInterfaceActionType::ToggleButton, FInputChord(EKeys::U, EModifierKey::Control) );
 	
 	UI_COMMAND(DragAnimSliderTool, "Drag Anim Slider Tool", "Drag the anim slider", EUserInterfaceActionType::Button, FInputChord(EKeys::U));
