Tweening: Added time offsetting, which shifts the curve in time by recomputing the key Y values thus maintaining key frame position.

diff --git a/Engine/Plugins/Animation/ControlRig/Source/ControlRigEditor/Private/EditMode/Tween/ControlRigTweenModels.cpp b/Engine/Plugins/Animation/ControlRig/Source/ControlRigEditor/Private/EditMode/Tween/ControlRigTweenModels.cpp
--- a/Engine/Plugins/Animation/ControlRig/Source/ControlRigEditor/Private/EditMode/Tween/ControlRigTweenModels.cpp
+++ b/Engine/Plugins/Animation/ControlRig/Source/ControlRigEditor/Private/EditMode/Tween/ControlRigTweenModels.cpp
@@ -4,6 +4,7 @@
 
 #include "ControlRigTweenModel.h"
 #include "Math/Abstraction/TweenRangeTemplates.h"
+#include "Math/Models/CurveTimeOffsetTweenModel.h"
 #include "Math/Models/EditorTweenModel.h"
 #include "Misc/SequencerUtils.h"
 
@@ -38,6 +39,14 @@ static TArray<TweeningUtilsEditor::FTweenModelUIEntry> MakeTweenModels(
 		Result.Emplace(MoveTemp(TweenModel), FTweenModelDisplayInfo(BlendFunction));
 	});
 
+	Result.Emplace(
+		MakeUnique<TEditorTweenModel<FCurveTimeOffsetTweenModel>>(
+			CurveEditorAttr,		// TTangentFlatteningTweenProxy arg
+			CurveEditorAttr			// FCurveTimeOffsetTweenModel args
+			),
+		FTweenModelDisplayInfo(EBlendFunction::TimeOffset)
+	);
+
 	return Result;
 }
 }
diff --git a/Engine/Plugins/Animation/TweeningUtils/Source/TweeningUtils/Private/Math/KeyBlendingFunctions.cpp b/Engine/Plugins/Animation/TweeningUtils/Source/TweeningUtils/Private/Math/KeyBlendingFunctions.cpp
--- a/Engine/Plugins/Animation/TweeningUtils/Source/TweeningUtils/Private/Math/KeyBlendingFunctions.cpp
+++ b/Engine/Plugins/Animation/TweeningUtils/Source/TweeningUtils/Private/Math/KeyBlendingFunctions.cpp
@@ -134,4 +134,44 @@ double Blend_SmoothRough(double InBlendValue, const FVector2d& InBeforeCurrent,
 		return InCurrent.Y + InBlendValue * (InCurrent.Y - NewValue);
 	}
 }
+
+double Blend_OffsetTime(
+	double InBlendValue, const FVector2d& InCurrent,
+	const FVector2d& InFirstBlended, const FVector2d& InLastBlended,
+	const FVector2d& InBeforeBlendRange, const FVector2d& InAfterBlendRange,
+	const TFunctionRef<double(double X)>& InEvaluate
+	)
+{
+	if (FMath::IsNearlyZero(InBlendValue))
+	{
+		return InCurrent.Y;
+	}
+
+	// A periodic function's period is the delta X at which y values start repeating. E.g. for sin(x), it's 2*PI.
+	const double Period = InLastBlended.X - InFirstBlended.X;
+	const double ShiftAmount = Period * InBlendValue;
+	const double ShiftedX = InCurrent.X - ShiftAmount; // For function g to shift f to the right by a, g(x) := f(x - a).
+
+	const auto BlendEdge = [&ShiftedX](const FVector2d& Border, const FVector2d& BeyondBorder)
+	{
+		const double ExceedAmount = ShiftedX - Border.X;
+		const double BlendEdgeToNext = BeyondBorder.X - Border.X;
+		const double BlendValue = FMath::IsNearlyZero(BlendEdgeToNext) ? 1.0 : FMath::Min(1.0, ExceedAmount / (BlendEdgeToNext));
+		const double RemappedBlendValue = BlendValue * 2.0 - 1.0; // [0,1] to [-1, 1]
+		return Blend_ControlsToTween(RemappedBlendValue, Border, BeyondBorder);
+	};
+
+	// If ShiftedX lies to the left or right of blended range of X values, use the Y value of the key to the left or right of the blended range, respectively.
+	// Previous implementations would wrap around (i.e. evaluating it as mathematically periodic function) but that does not make any sense for animators.
+	if (ShiftedX >= InLastBlended.X)
+	{
+		return BlendEdge(InLastBlended, InAfterBlendRange);
+	}
+	if (ShiftedX <= InFirstBlended.X)
+	{
+		return BlendEdge(InFirstBlended, InBeforeBlendRange);
+	}
+	
+	return InEvaluate(ShiftedX);
+}
 }
diff --git a/Engine/Plugins/Animation/TweeningUtils/Source/TweeningUtils/Public/Math/KeyBlendingFunctions.h b/Engine/Plugins/Animation/TweeningUtils/Source/TweeningUtils/Public/Math/KeyBlendingFunctions.h
--- a/Engine/Plugins/Animation/TweeningUtils/Source/TweeningUtils/Public/Math/KeyBlendingFunctions.h
+++ b/Engine/Plugins/Animation/TweeningUtils/Source/TweeningUtils/Public/Math/KeyBlendingFunctions.h
@@ -3,6 +3,7 @@
 #pragma once
 
 #include "Math/MathFwd.h"
+#include "Templates/FunctionFwd.h"
 
 /**
  * The blending functions are generally used to iterate through a contiguous set of keys.
@@ -72,9 +73,36 @@ TWEENINGUTILS_API double Blend_Ease(double InBlendValue,
 
 /**
  * Blend values of -1 or 1, push adjacent blended keys further together or apart, respectively. -1 averages keys out while 1 increases jumps.
+ * Softens the curve or makes it harsh. Smooth is useful for softening noise, as found in mocap animations.
  * @return The new Y value that the current key should have.
  */
 TWEENINGUTILS_API double Blend_SmoothRough(double InBlendValue, 
 	const FVector2d& InBeforeCurrent, const FVector2d& InCurrent, const FVector2d& InAfterCurrent
 	);
+
+/**
+ * Effectively shifts the curve to the left and right without actually changing the keys' X values: the Y values are recomputed to achieve the shift.
+ * 
+ * Blend values of -1 or 1 shift the range a period to the left or right, respectively; the period is in (InFirstBlendedX, InLastBlendedX).
+ * For example, imagine the keys formed a sin wave: a blend value of 0.5 would effectively make it cos wave.
+ * Keys for which the X value would be shifted out of the blend range have their Y value clamped.
+ *
+ * @param InBlendValue The blend value in [-1,1] that determines the shift amount (relative to InLastBlendedX - InFirstBlendedX).
+ * @param InCurrent The value of the current key to blend
+ * @param InFirstBlended The minimum of all keys being blended
+ * @param InLastBlended The maximum X of all keys being blended
+ * @param InBeforeBlendRange The value to return if the ShiftedX <= InFirstBlendedX.
+ *	Usually the X of the first key before the blended range; should be InFirstBlendedX if there is none. 
+ * @param InAfterBlendRange The value to return if the InLastBlendedX <= ShiftedX.
+ *	Usually the X of the first key before the blended range; should be InLastBlendedX if there is none. 
+ * @param InEvaluate Evaluates the function to shift; only needs to handle X in range InBlendInfo.FirstBlended <= X InBlendInfo.LastBlended.X. 
+ * 
+ * @return The new Y value that the current key should have.
+ */
+TWEENINGUTILS_API double Blend_OffsetTime(
+	double InBlendValue, const FVector2d& InCurrent,
+	const FVector2d& InFirstBlended, const FVector2d& InLastBlended,
+	const FVector2d& InBeforeBlendRange, const FVector2d& InAfterBlendRange,
+	const TFunctionRef<double(double X)>& InEvaluate
+	);
 }
diff --git a/Engine/Plugins/Animation/TweeningUtils/Source/TweeningUtilsEditor/Private/Math/Abstraction/KeyBlendingAbstraction.cpp b/Engine/Plugins/Animation/TweeningUtils/Source/TweeningUtilsEditor/Private/Math/Abstraction/KeyBlendingAbstraction.cpp
--- a/Engine/Plugins/Animation/TweeningUtils/Source/TweeningUtilsEditor/Private/Math/Abstraction/KeyBlendingAbstraction.cpp
+++ b/Engine/Plugins/Animation/TweeningUtils/Source/TweeningUtilsEditor/Private/Math/Abstraction/KeyBlendingAbstraction.cpp
@@ -60,7 +60,7 @@ struct FStaticBlendFunctionData
 
 static FStaticBlendFunctionData& GetFunctionData(EBlendFunction InFunction)
 {
-	static_assert(static_cast<int32>(EBlendFunction::Num) == 6, "Extend this array");
+	static_assert(static_cast<int32>(EBlendFunction::Num) == 7, "Extend this array");
 	constexpr int32 Num = static_cast<int32>(EBlendFunction::Num);
 	static FStaticBlendFunctionData Data[Num] =
 	{
@@ -103,8 +103,15 @@ static FStaticBlendFunctionData& GetFunctionData(EBlendFunction InFunction)
 			EBlendFunction::SmoothRough, FTweeningUtilsCommands::Get().SetTweenSmoothRough,
 			LOCTEXT("SmoothRough.Label", "Smooth Rough"),
 			LOCTEXT("SmoothRough.Abbreviation", "SR"),
-			LOCTEXT("SmoothRough.Description", "Make selected keys smooth or rough"),
+			LOCTEXT("SmoothRough.Description", "Push adjacent blended keys further together or apart. Smooth is useful for softening noise, like in mocap animations."),
 			TEXT("SmoothRough")
+		),
+		FStaticBlendFunctionData(
+			EBlendFunction::TimeOffset, FTweeningUtilsCommands::Get().SetTweenTimeOffset,
+			LOCTEXT("TimeOffset.Label", "Time Offset"),
+			LOCTEXT("TimeOffset.Abbreviation", "TO"),
+			LOCTEXT("TimeOffset.Description", "Shifts the curve to the left / right by changing the keys' Y values and maintaining frame position."),
+			TEXT("TimeOffset")
 		)
 	};
 
diff --git a/Engine/Plugins/Animation/TweeningUtils/Source/TweeningUtilsEditor/Private/Math/Models/CurveTimeOffsetTweenModel.cpp b/Engine/Plugins/Animation/TweeningUtils/Source/TweeningUtilsEditor/Private/Math/Models/CurveTimeOffsetTweenModel.cpp
new file mode 100644
--- /dev/null
+++ b/Engine/Plugins/Animation/TweeningUtils/Source/TweeningUtilsEditor/Private/Math/Models/CurveTimeOffsetTweenModel.cpp
@@ -0,0 +1,82 @@
+﻿// Copyright Epic Games, Inc. All Rights Reserved.
+
+#include "Math/Models/CurveTimeOffsetTweenModel.h"
+
+#include "SCurveEditorView.h"
+#include "Math/CurveBlending.h"
+#include "Math/KeyBlendingFunctions.h"
+
+namespace UE::TweeningUtilsEditor
+{
+FCurveTimeOffsetTweenModel::FCurveTimeOffsetTweenModel(TAttribute<TWeakPtr<FCurveEditor>> InWeakCurveEditor)
+	: WeakCurveEditor(MoveTemp(InWeakCurveEditor))
+{}
+
+void FCurveTimeOffsetTweenModel::StartBlendOperation()
+{
+	const TSharedPtr<FCurveEditor> CurveEditor = WeakCurveEditor.Get().Pin();
+	ContiguousKeySelection = CurveEditor ? FContiguousKeyMapping(*CurveEditor) : FContiguousKeyMapping{};
+
+	OriginalBlendedCurves.Empty(OriginalBlendedCurves.Num());
+	for (const TPair<FCurveModelID, FContiguousKeyMapping::FContiguousKeysArray>& Pair : ContiguousKeySelection.KeyMap)
+	{
+		const FCurveModelID CurveId = Pair.Key;
+		const FCurveModel* Curve = CurveEditor->FindCurve(CurveId);
+
+		// The curve model may not support CreateBufferedCurveCopy. In that case, it won't be blended. Too bad for it.
+		TUniquePtr<IBufferedCurveModel> BufferCurveCopy = Curve ? Curve->CreateBufferedCurveCopy() : nullptr;
+		if (BufferCurveCopy)
+		{
+			OriginalBlendedCurves.Add(CurveId, MoveTemp(BufferCurveCopy));
+		}
+	}
+}
+
+void FCurveTimeOffsetTweenModel::StopBlendOperation()
+{
+	ContiguousKeySelection.KeyMap.Empty(ContiguousKeySelection.KeyMap.Num());
+	OriginalBlendedCurves.Empty(OriginalBlendedCurves.Num());
+}
+
+void FCurveTimeOffsetTweenModel::BlendValues(float InNormalizedValue)
+{
+	const TSharedPtr<FCurveEditor> CurveEditor = WeakCurveEditor.Get().Pin();
+	if (!CurveEditor)
+	{
+		return;
+	}
+
+	const float ScaledBlendValue = ScaleBlendValue(InNormalizedValue); 
+	TweeningUtilsEditor::BlendCurves_BySingleKey(*CurveEditor, ContiguousKeySelection,
+		[this, ScaledBlendValue](
+			const FCurveModelID& CurveId,
+			const FContiguousKeyMapping::FContiguousKeysArray& AllBlendedKeys, const FContiguousKeys& CurrentBlendRange,
+			int32 Index
+		)
+	{
+		const FVector2d& CurrentKey = AllBlendedKeys.GetCurrent(CurrentBlendRange, Index);
+			
+		const double FallbackValue = CurrentKey.Y;
+		const TUniquePtr<IBufferedCurveModel>* OriginalCurve = OriginalBlendedCurves.Find(CurveId);
+		if (!OriginalCurve)
+		{
+			return FallbackValue;
+		}
+
+		const auto Evaluate = [&FallbackValue, &OriginalCurve](double X)
+		{
+			// Though it shouldn't fail normally, Evaluate can theoretically fail, in which case it does not write to OutputValue.
+			// So, init OutputValue to a defined value.
+			double OutputValue = FallbackValue; 
+			OriginalCurve->Get()->Evaluate(X, OutputValue);
+			return OutputValue;
+		};
+		
+		return TweeningUtils::Blend_OffsetTime(ScaledBlendValue, CurrentKey,
+			AllBlendedKeys.GetFirstInBlendRange(CurrentBlendRange), AllBlendedKeys.GetLastInBlendRange(CurrentBlendRange),
+			AllBlendedKeys.GetBeforeBlendRange(CurrentBlendRange), AllBlendedKeys.GetAfterBlendRange(CurrentBlendRange),
+			Evaluate
+			);
+	});
+}
+}
diff --git a/Engine/Plugins/Animation/TweeningUtils/Source/TweeningUtilsEditor/Private/TweeningUtilsCommands.cpp b/Engine/Plugins/Animation/TweeningUtils/Source/TweeningUtilsEditor/Private/TweeningUtilsCommands.cpp
--- a/Engine/Plugins/Animation/TweeningUtils/Source/TweeningUtilsEditor/Private/TweeningUtilsCommands.cpp
+++ b/Engine/Plugins/Animation/TweeningUtils/Source/TweeningUtilsEditor/Private/TweeningUtilsCommands.cpp
@@ -22,13 +22,14 @@ FTweeningUtilsCommands::FTweeningUtilsCommands()
 
 void FTweeningUtilsCommands::RegisterCommands()
 {
-	static_assert(static_cast<int32>(EBlendFunction::Num) == 6, "Add a command here");
-	UI_COMMAND(SetControlsToTween, "Tween ", "Sets the tween mode to Tween", EUserInterfaceActionType::RadioButton, FInputChord() );
-	UI_COMMAND(SetTweenPushPull, "Push / Pull", "Sets the tween mode to Push / Pull", EUserInterfaceActionType::RadioButton, FInputChord() );
-	UI_COMMAND(SetTweenBlendNeighbor, "Blend Neighbor", "Sets the tween mode to Blend Neighbour", EUserInterfaceActionType::RadioButton, FInputChord() );
-	UI_COMMAND(SetTweenBlendRelative, "Blend Relative", "Sets the tween mode to Blend Relative", EUserInterfaceActionType::RadioButton, FInputChord() );
-	UI_COMMAND(SetTweenBlendEase, "Blend Ease", "Sets the tween mode to Blend Ease", EUserInterfaceActionType::RadioButton, FInputChord() );
-	UI_COMMAND(SetTweenSmoothRough, "Smooth / Rough", "Sets the tween mode to Smooth / Rough", EUserInterfaceActionType::RadioButton, FInputChord() );
+	static_assert(static_cast<int32>(EBlendFunction::Num) == 7, "Add a command here");
+	UI_COMMAND(SetControlsToTween, "Tween ", "Interpolates between the previous and next keys", EUserInterfaceActionType::RadioButton, FInputChord() );
+	UI_COMMAND(SetTweenPushPull, "Push / Pull", "Push or pull the values to the interpolation between the previous and next keys", EUserInterfaceActionType::RadioButton, FInputChord() );
+	UI_COMMAND(SetTweenBlendNeighbor, "Blend Neighbor", "Blend to the next or previous values for selected keys", EUserInterfaceActionType::RadioButton, FInputChord() );
+	UI_COMMAND(SetTweenBlendRelative, "Blend Relative", "Blend relative to the next or previous value for selected keys", EUserInterfaceActionType::RadioButton, FInputChord() );
+	UI_COMMAND(SetTweenBlendEase, "Blend Ease", "Blend with an ease falloff to the next or previous value for selected keys", EUserInterfaceActionType::RadioButton, FInputChord() );
+	UI_COMMAND(SetTweenSmoothRough, "Smooth / Rough", "Push adjacent blended keys further together or apart. Smooth is useful for softening noise, like in mocap animations.", EUserInterfaceActionType::RadioButton, FInputChord() );
+	UI_COMMAND(SetTweenTimeOffset, "Time Offset", "Shifts the curve to the left / right by changing the keys' Y values and maintaining frame position.", EUserInterfaceActionType::RadioButton, FInputChord() );
 	UI_COMMAND(ToggleOvershootMode, "Overshoot mode", "Overshoot extends the available range over 100%", EUserInterfaceActionType::ToggleButton, FInputChord(EKeys::U, EModifierKey::Control) );
 	
 	UI_COMMAND(DragAnimSliderTool, "Drag Anim Slider Tool", "Drag the anim slider", EUserInterfaceActionType::Button, FInputChord(EKeys::U));
diff --git a/Engine/Plugins/Animation/TweeningUtils/Source/TweeningUtilsEditor/Private/TweeningUtilsStyle.cpp b/Engine/Plugins/Animation/TweeningUtils/Source/TweeningUtilsEditor/Private/TweeningUtilsStyle.cpp
--- a/Engine/Plugins/Animation/TweeningUtils/Source/TweeningUtilsEditor/Private/TweeningUtilsStyle.cpp
+++ b/Engine/Plugins/Animation/TweeningUtils/Source/TweeningUtilsEditor/Private/TweeningUtilsStyle.cpp
@@ -26,38 +26,43 @@ FTweeningUtilsStyle::FTweeningUtilsStyle()
 	const FVector2D Icon20x20(20.0f, 20.0f);
 	
 	// Tween - Slider colors
-	static_assert(static_cast<int32>(EBlendFunction::Num) == 6, "Add a color here");
+	static_assert(static_cast<int32>(EBlendFunction::Num) == 7, "Add a color here");
 	const FLinearColor ControlsToTweenTint(254.f/255.f, 155.f/255.f, 7.f/255.f);
 	const FLinearColor PushPullTint(107.f/255.f, 223.f/255.f, 151.f/255.f);
 	const FLinearColor BlendNeighborTint(133.f/255.f, 185.f/255.f, 73.f/255.f);
 	const FLinearColor BlendRelativeTint(39.f/255.f, 177.f/255.f, 240.f/255.f);
 	const FLinearColor BlendEaseTint(148.f/255.f, 57.f/255.f, 175.f/255.f);
 	const FLinearColor SmoothRoughTint(235.f/255.f, 64.f/255.f, 64.f/255.f);
+	const FLinearColor TimeOffsetTint(96.f/255.f, 65.f/255.f, 234.f/255.f);
 	Set("TweeningUtils.SetControlsToTween.Color", ControlsToTweenTint);
 	Set("TweeningUtils.SetTweenPushPull.Color", PushPullTint);
 	Set("TweeningUtils.SetTweenBlendNeighbor.Color", BlendNeighborTint);
 	Set("TweeningUtils.SetTweenBlendRelative.Color", BlendRelativeTint);
 	Set("TweeningUtils.SetTweenBlendEase.Color", BlendEaseTint);
 	Set("TweeningUtils.SetTweenSmoothRough.Color", SmoothRoughTint);
+	Set("TweeningUtils.SetTweenTimeOffset.Color", TimeOffsetTint);
+	
 	
 	// Tween - Commands
-	static_assert(static_cast<int32>(EBlendFunction::Num) == 6, "Add command style");
+	static_assert(static_cast<int32>(EBlendFunction::Num) == 7, "Add command style");
 	Set("TweeningUtils.SetControlsToTween", new IMAGE_BRUSH_SVG("Icons/ControlsToTween", Icon20x20, ControlsToTweenTint));
 	Set("TweeningUtils.SetTweenPushPull", new IMAGE_BRUSH_SVG("Icons/PushPull", Icon20x20, PushPullTint));
 	Set("TweeningUtils.SetTweenBlendNeighbor", new IMAGE_BRUSH_SVG("Icons/BlendNeighbor", Icon20x20, BlendNeighborTint));
 	Set("TweeningUtils.SetTweenBlendRelative", new IMAGE_BRUSH_SVG("Icons/BlendRelative", Icon20x20, BlendRelativeTint));
 	Set("TweeningUtils.SetTweenBlendEase", new IMAGE_BRUSH_SVG("Icons/BlendEase", Icon20x20, BlendEaseTint));
 	Set("TweeningUtils.SetTweenSmoothRough", new IMAGE_BRUSH_SVG("Icons/SmoothRough", Icon20x20, SmoothRoughTint));
+	Set("TweeningUtils.SetTweenTimeOffset", new IMAGE_BRUSH_SVG("Icons/TimeOffset", Icon20x20, TimeOffsetTint));
 	Set("TweeningUtils.ToggleOvershootMode", new IMAGE_BRUSH_SVG("Icons/OvershootMode", Icon20x20));
 	
 	// Tween - Icons (no tint)
-	static_assert(static_cast<int32>(EBlendFunction::Num) == 6, "Add untinted icon here");
+	static_assert(static_cast<int32>(EBlendFunction::Num) == 7, "Add untinted icon here");
 	Set("TweeningUtils.SetControlsToTween.Icon", new IMAGE_BRUSH_SVG("Icons/ControlsToTween", Icon20x20));
 	Set("TweeningUtils.SetTweenPushPull.Icon", new IMAGE_BRUSH_SVG("Icons/PushPull", Icon20x20));
 	Set("TweeningUtils.SetTweenBlendNeighbor.Icon", new IMAGE_BRUSH_SVG("Icons/BlendNeighbor", Icon20x20));
 	Set("TweeningUtils.SetTweenBlendRelative.Icon", new IMAGE_BRUSH_SVG("Icons/BlendRelative", Icon20x20));
 	Set("TweeningUtils.SetTweenBlendEase.Icon", new IMAGE_BRUSH_SVG("Icons/BlendEase", Icon20x20));
 	Set("TweeningUtils.SetTweenSmoothRough.Icon", new IMAGE_BRUSH_SVG("Icons/SmoothRough", Icon20x20));
+	Set("TweeningUtils.SetTweenTimeOffset.Icon", new IMAGE_BRUSH_SVG("Icons/TimeOffset", Icon20x20));
 	
 	// Tween - Slider visuals
 	Set("TweenSlider", FTweenSliderStyle());
diff --git a/Engine/Plugins/Animation/TweeningUtils/Source/TweeningUtilsEditor/Public/Math/Abstraction/KeyBlendingAbstraction.h b/Engine/Plugins/Animation/TweeningUtils/Source/TweeningUtilsEditor/Public/Math/Abstraction/KeyBlendingAbstraction.h
--- a/Engine/Plugins/Animation/TweeningUtils/Source/TweeningUtilsEditor/Public/Math/Abstraction/KeyBlendingAbstraction.h
+++ b/Engine/Plugins/Animation/TweeningUtils/Source/TweeningUtilsEditor/Public/Math/Abstraction/KeyBlendingAbstraction.h
@@ -20,6 +20,7 @@ namespace UE::TweeningUtilsEditor
 /** Util enum which is useful for UI code to abstract functions.*/
 enum class EBlendFunction : uint8
 {
+	// Simple blends, i.e. signature: double(const FKeyBlendInfo&).
 	BlendNeighbor,
 	PushPull,
 	BlendEase,
@@ -27,6 +28,9 @@ enum class EBlendFunction : uint8
 	BlendRelative,
 	SmoothRough,
 
+	// Complex blends
+	TimeOffset,
+
 	// Steps for adding a new function:
 	// 1. Add the function to KeyBlendingFunctions.h (consider splitting implementation to KeyBlendingFunctions_[YourFunc].cpp).
 	// 2. Extend this enum
diff --git a/Engine/Plugins/Animation/TweeningUtils/Source/TweeningUtilsEditor/Public/Math/Abstraction/TweenRangeTemplates.h b/Engine/Plugins/Animation/TweeningUtils/Source/TweeningUtilsEditor/Public/Math/Abstraction/TweenRangeTemplates.h
--- a/Engine/Plugins/Animation/TweeningUtils/Source/TweeningUtilsEditor/Public/Math/Abstraction/TweenRangeTemplates.h
+++ b/Engine/Plugins/Animation/TweeningUtils/Source/TweeningUtilsEditor/Public/Math/Abstraction/TweenRangeTemplates.h
@@ -25,13 +25,15 @@ double TweenRange(
 	double InBlendValue,
 	const FBlendRangesData& AllBlendedKeys, const FContiguousKeys& CurrentBlendRange,
 	int32 InCurrentKeyIndex
-); // Intentionally not implemented. This function is supposed to be specialiyed with the EBlendFunction types that support it.
+); // Intentionally not implemented. This function is supposed to be specialized with the EBlendFunction types that support it.
 
 /**
  * @return Whether this EBlendFunction supports simple tweening by using TweenRange.
  * Some blend functions require more arguments, like a cached curve, etc.
  */
 constexpr bool SupportsTweenRange(EBlendFunction BlendFunction);
+/** @return Number of functions for which SupportsTweenRange returns true.  */
+constexpr int32 NumBlendFunctionsSupportingTweenRange();
 	
 /**
  * Invokes InCallback for each EBlendFunction that TCurveTweener can be instantiated with.
@@ -47,7 +49,7 @@ constexpr bool SupportsTweenRange(EBlendFunction BlendFunction)
 {
 	// This function could use SFINAE but this way is easier to generate static_assert and remind developers to implement TweenRange with.
 	static_assert(
-		static_cast<int32>(EBlendFunction::Num) == 6,
+		static_cast<int32>(EBlendFunction::Num) == 7,
 		"Does the blend function you added support TweenRange? If so, specialize TweenRange and update this function."
 		); // TweenRange is supported if all inputs to your tween function can be determined by one of the FContiguousKeyRanges::GetX functions.
 
@@ -64,6 +66,19 @@ constexpr bool SupportsTweenRange(EBlendFunction BlendFunction)
 	}
 }
 
+constexpr int32 NumBlendFunctionsSupportingTweenRange()
+{
+	int32 Num = 0;
+	for (int32 FuncIdx = 0; FuncIdx < static_cast<int32>(EBlendFunction::Num); ++FuncIdx)
+	{
+		if (SupportsTweenRange(static_cast<EBlendFunction>(FuncIdx)))
+		{
+			++Num;
+		}
+	}
+	return Num;
+}
+
 template<>
 inline double TweenRange<EBlendFunction::ControlsToTween>(
 	double InBlendValue,
@@ -152,25 +167,16 @@ inline double TweenRange<EBlendFunction::SmoothRough>(
 		AllBlendedKeys.GetAfterCurrent(CurrentBlendRange, InCurrentKeyIndex)
 		);
 }
-	
-namespace Private
-{
-template <typename TCallback, int32... Indices>
-constexpr void ForEachCurveTweenableImpl(TCallback&& InCallback, TIntegerSequence<int32, Indices...>)
-{
-	// Fold expression... substitutes Indices for 1, 2, 3, etc., i.e. for each EBlendFunction entry.
-	((SupportsTweenRange(static_cast<EBlendFunction>(Indices))
-		? InCallback.template operator()<static_cast<EBlendFunction>(Indices)>() : void())
-	, ...);
-}
-}
 
 template <typename TCallback>
 constexpr void ForEachCurveTweenable(TCallback&& InCallback)
 {
-	Private::ForEachCurveTweenableImpl(
-		Forward<TCallback>(InCallback),
-		TMakeIntegerSequence<int32, static_cast<int32>(EBlendFunction::Num)>{}
-		);
+	static_assert(NumBlendFunctionsSupportingTweenRange() == 6, "Extend this function with the enum entry you added.");
+	InCallback.template operator()<EBlendFunction::BlendNeighbor>();
+	InCallback.template operator()<EBlendFunction::PushPull>();
+	InCallback.template operator()<EBlendFunction::BlendEase>();
+	InCallback.template operator()<EBlendFunction::ControlsToTween>();
+	InCallback.template operator()<EBlendFunction::BlendRelative>();
+	InCallback.template operator()<EBlendFunction::SmoothRough>();
 }
 }
\ No newline at end of file
diff --git a/Engine/Plugins/Animation/TweeningUtils/Source/TweeningUtilsEditor/Public/Math/Models/CurveTimeOffsetTweenModel.h b/Engine/Plugins/Animation/TweeningUtils/Source/TweeningUtilsEditor/Public/Math/Models/CurveTimeOffsetTweenModel.h
new file mode 100644
--- /dev/null
+++ b/Engine/Plugins/Animation/TweeningUtils/Source/TweeningUtilsEditor/Public/Math/Models/CurveTimeOffsetTweenModel.h
@@ -0,0 +1,41 @@
+﻿// Copyright Epic Games, Inc. All Rights Reserved.
+
+#pragma once
+
+#include "IBufferedCurveModel.h"
+#include "TweenModel.h"
+#include "Math/ContiguousKeyMapping.h"
+#include "Misc/Attribute.h"
+#include "Templates/SharedPointer.h"
+#include "Templates/UnrealTemplate.h"
+
+class FCurveEditor;
+
+namespace UE::TweeningUtilsEditor
+{
+/** Recomputes all keys output values such that the curve is effectively shifted to left and right without modifying the keys' input values. */
+class TWEENINGUTILSEDITOR_API FCurveTimeOffsetTweenModel : public FTweenModel, public FNoncopyable
+{
+public:
+
+	explicit FCurveTimeOffsetTweenModel(TAttribute<TWeakPtr<FCurveEditor>> InWeakCurveEditor);
+	
+	//~ Begin FTweenModel Interface
+	virtual void StartBlendOperation() override;
+	virtual void StopBlendOperation() override;
+	virtual void BlendValues(float InNormalizedValue) override;
+	//~ Begin FTweenModel Interface
+
+protected:
+	
+	/** The curve editor on which to tween the curves. */
+	const TAttribute<TWeakPtr<FCurveEditor>> WeakCurveEditor;
+
+	/** Created in StartBlendOperation and used for the entirety of the blend operation. */
+	FContiguousKeyMapping ContiguousKeySelection;
+
+	/** The state of the curves before they were blended. Created in StartBlendOperation. Used to evaluate original curve value during time shift. */
+	TMap<FCurveModelID, TUniquePtr<IBufferedCurveModel>> OriginalBlendedCurves;
+};
+}
+
diff --git a/Engine/Plugins/Animation/TweeningUtils/Source/TweeningUtilsEditor/Public/Math/Models/CurveTweenModel.h b/Engine/Plugins/Animation/TweeningUtils/Source/TweeningUtilsEditor/Public/Math/Models/CurveTweenModel.h
--- a/Engine/Plugins/Animation/TweeningUtils/Source/TweeningUtilsEditor/Public/Math/Models/CurveTweenModel.h
+++ b/Engine/Plugins/Animation/TweeningUtils/Source/TweeningUtilsEditor/Public/Math/Models/CurveTweenModel.h
@@ -14,11 +14,12 @@ class FCurveEditor;
 namespace UE::TweeningUtilsEditor
 {
 /** Tweens the curves of the curve editor's selection. */
-template<EBlendFunction>
+template<EBlendFunction Func>
 class TCurveTweenModel : public FTweenModel
 {
+	static_assert(SupportsTweenRange(Func), "This EBlendFunction does not support TweenRange<EBlendFunction>");
 public:
-
+	
 	explicit TCurveTweenModel(TWeakPtr<FCurveEditor> InWeakCurveEditor);
 
 	/** @return Whether it makes sense to call StartBlendOperation, etc. */
diff --git a/Engine/Plugins/Animation/TweeningUtils/Source/TweeningUtilsEditor/Public/TweeningUtilsCommands.h b/Engine/Plugins/Animation/TweeningUtils/Source/TweeningUtilsEditor/Public/TweeningUtilsCommands.h
--- a/Engine/Plugins/Animation/TweeningUtils/Source/TweeningUtilsEditor/Public/TweeningUtilsCommands.h
+++ b/Engine/Plugins/Animation/TweeningUtils/Source/TweeningUtilsEditor/Public/TweeningUtilsCommands.h
@@ -16,6 +16,7 @@ public:
 	TSharedPtr<FUICommandInfo> SetTweenBlendRelative;
 	TSharedPtr<FUICommandInfo> SetTweenBlendEase;
 	TSharedPtr<FUICommandInfo> SetTweenSmoothRough;
+	TSharedPtr<FUICommandInfo> SetTweenTimeOffset;
 	TSharedPtr<FUICommandInfo> ToggleOvershootMode;
 
 	TSharedPtr<FUICommandInfo> DragAnimSliderTool;
diff --git a/Engine/Plugins/Editor/CurveEditorTools/Source/CurveEditorTools/Private/Extensions/Tweening/CurveEditorTweenModels.cpp b/Engine/Plugins/Editor/CurveEditorTools/Source/CurveEditorTools/Private/Extensions/Tweening/CurveEditorTweenModels.cpp
--- a/Engine/Plugins/Editor/CurveEditorTools/Source/CurveEditorTools/Private/Extensions/Tweening/CurveEditorTweenModels.cpp
+++ b/Engine/Plugins/Editor/CurveEditorTools/Source/CurveEditorTools/Private/Extensions/Tweening/CurveEditorTweenModels.cpp
@@ -3,6 +3,7 @@
 #include "CurveEditorTweenModels.h"
 
 #include "ContextAwareTweeningModel.h"
+#include "Math/Models/CurveTimeOffsetTweenModel.h"
 #include "Math/Models/EditorTweenModel.h"
 
 namespace UE::CurveEditorTools
@@ -16,6 +17,11 @@ static TArray<TweeningUtilsEditor::FTweenModelUIEntry> MakeTweenModels(TWeakPtr<
 	TArray<FTweenModelUIEntry> Result;
 	ForEachCurveTweenable([&InCurveEditor, &Result]<EBlendFunction TBlendFunction>
 	{
+		if constexpr (!SupportsTweenRange(TBlendFunction))
+		{
+			return;
+		}
+		
 		using FModel = TEditorTweenModel<TContextAwareTweeningModel<TBlendFunction>>;
 
 		TUniquePtr<FModel> TweenModel = MakeUnique<FModel>(
@@ -24,6 +30,15 @@ static TArray<TweeningUtilsEditor::FTweenModelUIEntry> MakeTweenModels(TWeakPtr<
 			);
 		Result.Emplace(MoveTemp(TweenModel), FTweenModelDisplayInfo(TBlendFunction));
 	});
+	static_assert(!SupportsTweenRange(EBlendFunction::TimeOffset));
+	
+	Result.Emplace(
+		MakeUnique<TEditorTweenModel<FCurveTimeOffsetTweenModel>>(
+			InCurveEditor,		// TTangentFlatteningTweenProxy arg
+			InCurveEditor		// FCurveTimeOffsetTweenModel args
+			),
+		FTweenModelDisplayInfo(EBlendFunction::TimeOffset)
+	);
 
 	return Result;
 }
