Curve Editor: New icons for toolbar:
- The 7 tween functions have new icons
- Overshoot mode
- Curve Settings
- Key combo button: No keys selected & mixed keys

diff --git a/Engine/Plugins/Animation/TweeningUtils/Source/TweeningUtilsEditor/Private/TweeningUtilsStyle.cpp b/Engine/Plugins/Animation/TweeningUtils/Source/TweeningUtilsEditor/Private/TweeningUtilsStyle.cpp
--- a/Engine/Plugins/Animation/TweeningUtils/Source/TweeningUtilsEditor/Private/TweeningUtilsStyle.cpp
+++ b/Engine/Plugins/Animation/TweeningUtils/Source/TweeningUtilsEditor/Private/TweeningUtilsStyle.cpp
@@ -45,24 +45,24 @@ FTweeningUtilsStyle::FTweeningUtilsStyle()
 	
 	// Tween - Commands
 	static_assert(static_cast<int32>(EBlendFunction::Num) == 7, "Add command style");
-	Set("TweeningUtils.SetControlsToTween", new IMAGE_BRUSH_SVG("Icons/ControlsToTween", Icon20x20, ControlsToTweenTint));
-	Set("TweeningUtils.SetTweenPushPull", new IMAGE_BRUSH_SVG("Icons/PushPull", Icon20x20, PushPullTint));
-	Set("TweeningUtils.SetTweenBlendNeighbor", new IMAGE_BRUSH_SVG("Icons/BlendNeighbor", Icon20x20, BlendNeighborTint));
-	Set("TweeningUtils.SetTweenBlendRelative", new IMAGE_BRUSH_SVG("Icons/BlendRelative", Icon20x20, BlendRelativeTint));
-	Set("TweeningUtils.SetTweenBlendEase", new IMAGE_BRUSH_SVG("Icons/BlendEase", Icon20x20, BlendEaseTint));
-	Set("TweeningUtils.SetTweenSmoothRough", new IMAGE_BRUSH_SVG("Icons/SmoothRough", Icon20x20, SmoothRoughTint));
-	Set("TweeningUtils.SetTweenTimeOffset", new IMAGE_BRUSH_SVG("Icons/TimeOffset", Icon20x20, TimeOffsetTint));
-	Set("TweeningUtils.ToggleOvershootMode", new IMAGE_BRUSH_SVG("Icons/OvershootMode", Icon20x20));
+	Set("TweeningUtils.SetControlsToTween", new IMAGE_BRUSH_SVG("Icons/CurveTween_16", Icon20x20, ControlsToTweenTint));
+	Set("TweeningUtils.SetTweenPushPull", new IMAGE_BRUSH_SVG("Icons/CurvePushPull_16", Icon20x20, PushPullTint));
+	Set("TweeningUtils.SetTweenBlendNeighbor", new IMAGE_BRUSH_SVG("Icons/CurveBlendNeighbor_16", Icon20x20, BlendNeighborTint));
+	Set("TweeningUtils.SetTweenBlendRelative", new IMAGE_BRUSH_SVG("Icons/CurveMoveRelative_16", Icon20x20, BlendRelativeTint));
+	Set("TweeningUtils.SetTweenBlendEase", new IMAGE_BRUSH_SVG("Icons/CurveBlendEase_16", Icon20x20, BlendEaseTint));
+	Set("TweeningUtils.SetTweenSmoothRough", new IMAGE_BRUSH_SVG("Icons/CurveSmoothRough_16", Icon20x20, SmoothRoughTint));
+	Set("TweeningUtils.SetTweenTimeOffset", new IMAGE_BRUSH_SVG("Icons/CurveTimeOffset_16", Icon20x20, TimeOffsetTint));
+	Set("TweeningUtils.ToggleOvershootMode", new IMAGE_BRUSH_SVG("Icons/SliderOvershoot_20", Icon20x20));
 	
 	// Tween - Icons (no tint)
 	static_assert(static_cast<int32>(EBlendFunction::Num) == 7, "Add untinted icon here");
-	Set("TweeningUtils.SetControlsToTween.Icon", new IMAGE_BRUSH_SVG("Icons/ControlsToTween", Icon20x20));
-	Set("TweeningUtils.SetTweenPushPull.Icon", new IMAGE_BRUSH_SVG("Icons/PushPull", Icon20x20));
-	Set("TweeningUtils.SetTweenBlendNeighbor.Icon", new IMAGE_BRUSH_SVG("Icons/BlendNeighbor", Icon20x20));
-	Set("TweeningUtils.SetTweenBlendRelative.Icon", new IMAGE_BRUSH_SVG("Icons/BlendRelative", Icon20x20));
-	Set("TweeningUtils.SetTweenBlendEase.Icon", new IMAGE_BRUSH_SVG("Icons/BlendEase", Icon20x20));
-	Set("TweeningUtils.SetTweenSmoothRough.Icon", new IMAGE_BRUSH_SVG("Icons/SmoothRough", Icon20x20));
-	Set("TweeningUtils.SetTweenTimeOffset.Icon", new IMAGE_BRUSH_SVG("Icons/TimeOffset", Icon20x20));
+	Set("TweeningUtils.SetControlsToTween.Icon", new IMAGE_BRUSH_SVG("Icons/CurveTween_16", Icon20x20));
+	Set("TweeningUtils.SetTweenPushPull.Icon", new IMAGE_BRUSH_SVG("Icons/CurvePushPull_16", Icon20x20));
+	Set("TweeningUtils.SetTweenBlendNeighbor.Icon", new IMAGE_BRUSH_SVG("Icons/CurveBlendNeighbor_16", Icon20x20));
+	Set("TweeningUtils.SetTweenBlendRelative.Icon", new IMAGE_BRUSH_SVG("Icons/CurveMoveRelative_16", Icon20x20));
+	Set("TweeningUtils.SetTweenBlendEase.Icon", new IMAGE_BRUSH_SVG("Icons/CurveBlendEase_16", Icon20x20));
+	Set("TweeningUtils.SetTweenSmoothRough.Icon", new IMAGE_BRUSH_SVG("Icons/CurveSmoothRough_16", Icon20x20));
+	Set("TweeningUtils.SetTweenTimeOffset.Icon", new IMAGE_BRUSH_SVG("Icons/CurveTimeOffset_16", Icon20x20));
 	
 	// Tween - Slider visuals
 	Set("TweenSlider", FTweenSliderStyle());
diff --git a/Engine/Plugins/Animation/TweeningUtils/Source/TweeningUtilsEditor/Private/Widgets/MVC/TweenToolbarController.cpp b/Engine/Plugins/Animation/TweeningUtils/Source/TweeningUtilsEditor/Private/Widgets/MVC/TweenToolbarController.cpp
--- a/Engine/Plugins/Animation/TweeningUtils/Source/TweeningUtilsEditor/Private/Widgets/MVC/TweenToolbarController.cpp
+++ b/Engine/Plugins/Animation/TweeningUtils/Source/TweeningUtilsEditor/Private/Widgets/MVC/TweenToolbarController.cpp
@@ -15,6 +15,8 @@
 
 namespace UE::TweeningUtilsEditor
 {
+static const FName BaseOvershootBrushName(TEXT("OvershootMode"));
+	
 class FOvershootButtonStyleHack : public FSlateStyleSet
 {
 public:
@@ -28,8 +30,8 @@ public:
 		FSlateStyleSet::SetCoreContentRoot(EngineEditorSlateDir);
 		
 		const FVector2D Icon20x20(20.0f, 20.0f);
-		FSlateBrush* BaseBrush = new IMAGE_BRUSH_SVG("Icons/OvershootMode", Icon20x20);
-		Set(TEXT("OvershootMode"), BaseBrush);
+		FSlateBrush* BaseBrush = new IMAGE_BRUSH_SVG("Icons/SliderOvershoot_20", Icon20x20);
+		Set(BaseOvershootBrushName, BaseBrush);
 		InTweenFunctions.ForEachModel([this, &InTweenFunctions, &BaseBrush](const FTweenModel& Function) 
 		{
 			FSlateBrush* Brush = new FSlateBrush(*BaseBrush);
@@ -186,7 +188,7 @@ FSlateIcon FTweenToolbarController::GetOvershootModeIcon() const
 {
 	return FSlateIcon(
 		OverrideStyle->GetStyleSetName(),
-		IsOvershootModeEnabled() ? TEXT("OvershootMode") : *OverrideStyle->GetOvershootButtonStyleName(*TweenModels, *SelectedTweenModel)
+		IsOvershootModeEnabled() ? BaseOvershootBrushName : FName(*OverrideStyle->GetOvershootButtonStyleName(*TweenModels, *SelectedTweenModel))
 		);
 }
 
diff --git a/Engine/Source/Editor/EditorStyle/Private/StarshipStyle.cpp b/Engine/Source/Editor/EditorStyle/Private/StarshipStyle.cpp
--- a/Engine/Source/Editor/EditorStyle/Private/StarshipStyle.cpp
+++ b/Engine/Source/Editor/EditorStyle/Private/StarshipStyle.cpp
@@ -7719,8 +7719,8 @@ void FStarshipEditorStyle::FStyle::SetupToolkitStyles()
 		Set("GenericCurveEditor.InterpolationToggleWeighted", new IMAGE_BRUSH_SVG("Starship/GenericCurveEditor/TangentsWeighted", Icon20x20));
 		Set("GenericCurveEditor.InterpolationLinear", new IMAGE_BRUSH_SVG("Starship/GenericCurveEditor/TangentsLinear", Icon20x20));
 		Set("GenericCurveEditor.InterpolationConstant", new IMAGE_BRUSH_SVG("Starship/GenericCurveEditor/TangentsConstant", Icon20x20));
-		Set("GenericCurveEditor.InterpolationMixed", new IMAGE_BRUSH_SVG("Starship/GenericCurveEditor/TangentsMixed", Icon20x20));
-		Set("GenericCurveEditor.InterpolationNoSelection", new IMAGE_BRUSH_SVG("Starship/GenericCurveEditor/TangentsNoSelection", Icon20x20));
+		Set("GenericCurveEditor.InterpolationMixed", new IMAGE_BRUSH_SVG("Starship/GenericCurveEditor/CurveEditorTangentMultiSelected_20", Icon20x20));
+		Set("GenericCurveEditor.InterpolationNoSelection", new IMAGE_BRUSH_SVG("Starship/GenericCurveEditor/CurveEditorTangentUnselected_20", Icon20x20));
 
 		// Tangent Modifications
 		Set("GenericCurveEditor.FlattenTangents", new IMAGE_BRUSH_SVG("Starship/GenericCurveEditor/MiscFlatten", Icon20x20));
@@ -7771,7 +7771,7 @@ void FStarshipEditorStyle::FStyle::SetupToolkitStyles()
 		
 		Set("GenericCurveEditor.Select", new IMAGE_BRUSH("GenericCurveEditor/Icons/SelectButton", Icon16x16));
 		
-		Set("GenericCurveEditor.Curves", new IMAGE_BRUSH_SVG("Starship/GenericCurveEditor/GraphViewAbsolute_20", Icon20x20)); 
+		Set("GenericCurveEditor.Curves", new IMAGE_BRUSH_SVG("Starship/GenericCurveEditor/CurveSettings_20", Icon20x20)); 
 	}
 
 	// Generic Curve Editor Tools
