Curve Editor: Fix issues with STweenSlider in menu.
- Tweening now updates the viewport when the slider is in a menu: Slate's throttling prevented the viewport from updating.
- Adjust painting code to correctly fill all allotted space. This fixes the tween button not aligning with the mouse cursor while dragging while the widget is in a menu.

diff --git a/Engine/Plugins/Animation/TweeningUtils/Source/TweeningUtilsEditor/Private/Widgets/MVC/TweenToolbarController.cpp b/Engine/Plugins/Animation/TweeningUtils/Source/TweeningUtilsEditor/Private/Widgets/MVC/TweenToolbarController.cpp
--- a/Engine/Plugins/Animation/TweeningUtils/Source/TweeningUtilsEditor/Private/Widgets/MVC/TweenToolbarController.cpp
+++ b/Engine/Plugins/Animation/TweeningUtils/Source/TweeningUtilsEditor/Private/Widgets/MVC/TweenToolbarController.cpp
@@ -136,7 +136,7 @@ void FTweenToolbarController::SetSelectedTweenModel(FTweenModel& InTweenModel)
 
 TSharedRef<SWidget> FTweenToolbarController::MakeTweenModeMenu() const
 {
-	FMenuBuilder MenuBuilder(true, CommandList);
+	FMenuBuilder MenuBuilder(false, CommandList);
 
 	TweenModels->ForEachModel([this, &MenuBuilder](const FTweenModel& InTweenModel)
 	{
diff --git a/Engine/Plugins/Animation/TweeningUtils/Source/TweeningUtilsEditor/Private/Widgets/STweenSlider.cpp b/Engine/Plugins/Animation/TweeningUtils/Source/TweeningUtilsEditor/Private/Widgets/STweenSlider.cpp
--- a/Engine/Plugins/Animation/TweeningUtils/Source/TweeningUtilsEditor/Private/Widgets/STweenSlider.cpp
+++ b/Engine/Plugins/Animation/TweeningUtils/Source/TweeningUtilsEditor/Private/Widgets/STweenSlider.cpp
@@ -12,6 +12,14 @@
 
 namespace UE::TweeningUtilsEditor
 {
+class FScopedDisableThrottle : public FNoncopyable
+{
+public:
+
+	FScopedDisableThrottle() { FSlateThrottleManager::Get().DisableThrottle(true); }
+	~FScopedDisableThrottle(){ FSlateThrottleManager::Get().DisableThrottle(false); }
+};
+	
 struct FSliderWidgetData
 {
 	/** Contains all values required for drawing the slider. */
@@ -32,10 +40,15 @@ struct FSliderWidgetData
 
 		/** Used to determine whether the slider value changes this tick. */
 		float LastSliderPosition = 0.5f;
+
+		/** Set when a drag is detected. Makes sure that dragging the slider updates the viewport in realtime if the slide is in a menu. */
+		TOptional<FScopedDisableThrottle> ScopedPreventThrottle;
 		
 		explicit FDelayedDragData(const FVector2D& InInitialPosition, const FKey& InEffectiveKey)
 			: DelayedDrag(InInitialPosition, InEffectiveKey)
 		{}
+
+		void EnterResponsiveMode() { ScopedPreventThrottle.Emplace(); }
 	};
 	/** Data valid while the mouse button is down. */
 	TOptional<FDelayedDragData> DragData;
@@ -146,7 +159,7 @@ FReply STweenSlider::OnMouseButtonDown(const FGeometry& InGeometry, const FPoint
 		Pimpl->DragData->DelayedDrag.ForceDragStart();
 		UpdateLastMousePositionOnSlider(InGeometry, InMouseEvent.GetScreenSpacePosition());
 		Pimpl->TargetSliderPosition = Pimpl->LastMousePositionOnSliderBar;
-		OnSliderDragStartedDelegate.ExecuteIfBound();
+		OnStartDragOp();
 	}
 	
 	return FReply::Handled().CaptureMouse(SharedThis(this)).PreventThrottling();
@@ -168,7 +181,7 @@ FReply STweenSlider::OnMouseButtonUp(const FGeometry& InGeometry, const FPointer
 	if (Pimpl->DragData)
 	{
 		Pimpl->DragData.Reset();
-		OnSliderDragStoppedDelegate.ExecuteIfBound();
+		OnSliderDragStoppedDelegate.ExecuteIfBound();;
 		
 		Pimpl->AnimatedCurrentSliderPosition = 0.5f;
 		Pimpl->TargetSliderPosition = 0.5f;
@@ -198,7 +211,7 @@ FReply STweenSlider::OnMouseMove(const FGeometry& InGeometry, const FPointerEven
 		const bool bHasStarted = Pimpl->DragData->DelayedDrag.AttemptDragStart(InMouseEvent);
 		if (bHasStarted)
 		{
-			OnSliderDragStartedDelegate.ExecuteIfBound();
+			OnStartDragOp();
 		}
 
 		if (IsDragging())
@@ -209,7 +222,8 @@ FReply STweenSlider::OnMouseMove(const FGeometry& InGeometry, const FPointerEven
 		return FReply::Handled()
 			.CaptureMouse(SharedThis(this))
 			// Locking the mouse gives the user instant feedback about having reached the end of the slider range.
-			.LockMouseToWidget(SharedThis(this));
+			.LockMouseToWidget(SharedThis(this))
+			.PreventThrottling();
 	}
 
 	// Hover state needs to updated for OnPaint (don't update it in OnPaint because Slate shifts InGeometry in OnPaint, which messes up hovering computations)
@@ -335,6 +349,12 @@ bool STweenSlider::IsDragging() const
 	return Pimpl->DragData && Pimpl->DragData->DelayedDrag.IsDragging(); 
 }
 
+void STweenSlider::OnStartDragOp() const
+{
+	Pimpl->DragData->EnterResponsiveMode();
+	OnSliderDragStartedDelegate.ExecuteIfBound();
+}
+
 void STweenSlider::HandlePickValueOfCurrentlyHoveredPoint(const FGeometry& InGeometry, const FPointerEvent& InMouseEvent) const
 {
 	FTweenWidgetArgs& WidgetArgs = Pimpl->WidgetArgs;
@@ -349,9 +369,9 @@ void STweenSlider::HandlePickValueOfCurrentlyHoveredPoint(const FGeometry& InGeo
 	const int32 HoverIndex = HoverState.HoveredPointIndex.Get(INDEX_NONE);
 	if (PointValues.IsValidIndex(HoverIndex))
 	{
-		OnPointValuePickedDelegate.ExecuteIfBound(
-			AsymmetricToSymmetric(PointValues[HoverIndex])
-			);
+		// TODO UE-282985: If this widget is in a menu, then Slate will throttle and prevent the viewport from updating.
+		// Simply disable throttling for 1 frame does not work.
+		OnPointValuePickedDelegate.ExecuteIfBound(AsymmetricToSymmetric(PointValues[HoverIndex]));
 	}
 }
 
diff --git a/Engine/Plugins/Animation/TweeningUtils/Source/TweeningUtilsEditor/Private/Widgets/TweenSliderDrawUtils.cpp b/Engine/Plugins/Animation/TweeningUtils/Source/TweeningUtilsEditor/Private/Widgets/TweenSliderDrawUtils.cpp
--- a/Engine/Plugins/Animation/TweeningUtils/Source/TweeningUtilsEditor/Private/Widgets/TweenSliderDrawUtils.cpp
+++ b/Engine/Plugins/Animation/TweeningUtils/Source/TweeningUtilsEditor/Private/Widgets/TweenSliderDrawUtils.cpp
@@ -8,6 +8,7 @@
 
 namespace UE::TweeningUtilsEditor
 {
+	
 static FVector2D ComputeButtonSize(const FTweenWidgetArgs& InWidgetArgs)
 {
 	const FSlateBrush* SliderIconBrush = InWidgetArgs.SliderIconAttr.Get();
@@ -21,11 +22,24 @@ static FVector2D ComputeBarOffset(const FTweenWidgetArgs& InWidgetArgs)
 	const FVector2D HalfButtonOffset { ComputeButtonSize(InWidgetArgs).X / 2.0, 0.0 };
 	return HalfButtonOffset;
 }
+
+static FVector2D GetBarDimensions(const FGeometry& AllottedGeometry, const FTweenWidgetArgs& InWidgetArgs)
+{
+	const FVector2D& Dimensions = InWidgetArgs.Style->BarDimensions;
+	
+	// This handles widgets slots set to Auto and Filled. At least take the amount of space defined by the style.
+	// If there is more horizontal space available, fill that space out.
+	// Do not fill vertically because it looks broken: use the style value.
+	const FVector2D BarSize = FVector2D{ FMath::Max(Dimensions.X, AllottedGeometry.GetLocalSize().X), Dimensions.Y };
+	
+	return BarSize
+		// Account for the offset (left) and inset (right); there needs to be enough space to fully drag the button to the left and right.
+		- 2 * ComputeBarOffset(InWidgetArgs);
+}
 	
 void GetBarGeometry(const FGeometry& AllottedGeometry, const FTweenWidgetArgs& InWidgetArgs, FGeometry& OutBarArea)
 {
-	const FVector2D& Dimensions = InWidgetArgs.Style->BarDimensions;
-	const FVector2D BarSize = FVector2D{ Dimensions.X, Dimensions.Y };
+	const FVector2D BarSize = GetBarDimensions(AllottedGeometry, InWidgetArgs);
 	OutBarArea = AllottedGeometry.MakeChild(
 		BarSize, FSlateLayoutTransform(ComputeBarOffset(InWidgetArgs) + FVector2D{ 0.f, AllottedGeometry.Size.Y / 2.f - BarSize.Y / 2.f })
 		);
@@ -35,7 +49,7 @@ static FGeometry MakeGeometryOnBar(
 	float InNormalizedBarPos, const FTweenWidgetArgs& InWidgetArgs, const FVector2D& InSize, const FGeometry& InAllottedGeometry
 	)
 {
-	const FVector2D& Dimensions = InWidgetArgs.Style->BarDimensions;
+	const FVector2D Dimensions = GetBarDimensions(InAllottedGeometry, InWidgetArgs);
 	const float HorizontalOffset = Dimensions.X * InNormalizedBarPos - (InSize.X / 2.f);
 	return InAllottedGeometry.MakeChild(
 		InSize,
@@ -49,7 +63,7 @@ static FGeometry MakeGeometryOnBar(
 	
 void GetDragValueIndicationGeometry(const FGeometry& InAllottedGeometry, const FTweenWidgetArgs& InWidgetArgs, float InSliderPosition, FGeometry& OutDragValueIndication)
 {
-	const FVector2D& Dimensions = InWidgetArgs.Style->BarDimensions;
+	const FVector2D Dimensions = GetBarDimensions(InAllottedGeometry, InWidgetArgs);
 	
 	const float RelativeWidth = InSliderPosition - 0.5f;
 	const float AbsoluteWidth = FMath::Abs(RelativeWidth) * Dimensions.X;
diff --git a/Engine/Plugins/Animation/TweeningUtils/Source/TweeningUtilsEditor/Public/Widgets/STweenSlider.h b/Engine/Plugins/Animation/TweeningUtils/Source/TweeningUtilsEditor/Public/Widgets/STweenSlider.h
--- a/Engine/Plugins/Animation/TweeningUtils/Source/TweeningUtilsEditor/Public/Widgets/STweenSlider.h
+++ b/Engine/Plugins/Animation/TweeningUtils/Source/TweeningUtilsEditor/Public/Widgets/STweenSlider.h
@@ -113,6 +113,8 @@ private:
 	bool IsMouseDown() const;
 	bool IsDragging() const;
 
+	void OnStartDragOp() const;
+
 	/** If the user is currently hovering a point, set the blend value to that point's value.  */
 	void HandlePickValueOfCurrentlyHoveredPoint(const FGeometry& InGeometry, const FPointerEvent& InMouseEvent) const;
 
