Tweening Utils: Add IsRegistered check in destructor. This handles external modules keeping alive FTweenToolbarController past the life-time of the TweeningUtils module.

#rnx

diff --git a/Engine/Plugins/Animation/TweeningUtils/Source/TweeningUtilsEditor/Private/Widgets/MVC/TweenToolbarController.cpp b/Engine/Plugins/Animation/TweeningUtils/Source/TweeningUtilsEditor/Private/Widgets/MVC/TweenToolbarController.cpp
--- a/Engine/Plugins/Animation/TweeningUtils/Source/TweeningUtilsEditor/Private/Widgets/MVC/TweenToolbarController.cpp
+++ b/Engine/Plugins/Animation/TweeningUtils/Source/TweeningUtilsEditor/Private/Widgets/MVC/TweenToolbarController.cpp
@@ -93,8 +93,12 @@ FTweenToolbarController::~FTweenToolbarController()
 	{
 		CommandList->UnmapAction(TweenModels->GetCommandForModel(InTweenModel));
 	});
-	
-	CommandList->UnmapAction(FTweeningUtilsCommands::Get().ToggleOvershootMode);
+
+	// If the command is invalid, we're shutting down and this doesn't need to be run.
+	if (FTweeningUtilsCommands::IsRegistered())
+	{
+		CommandList->UnmapAction(FTweeningUtilsCommands::Get().ToggleOvershootMode);
+	}
 }
 
 void FTweenToolbarController::AddToToolbar(FToolBarBuilder& ToolBarBuilder, FMakeWidgetArgs InArgs) const
