Curve Editor: Fix typo in GetBeforeCurrent. This fixes the smooth / rough function.

#Synced-CL 51259554

diff --git a/Engine/Plugins/Animation/TweeningUtils/Source/TweeningUtilsEditor/Public/Math/ContiguousKeyMapping.h b/Engine/Plugins/Animation/TweeningUtils/Source/TweeningUtilsEditor/Public/Math/ContiguousKeyMapping.h
--- a/Engine/Plugins/Animation/TweeningUtils/Source/TweeningUtilsEditor/Public/Math/ContiguousKeyMapping.h
+++ b/Engine/Plugins/Animation/TweeningUtils/Source/TweeningUtilsEditor/Public/Math/ContiguousKeyMapping.h
@@ -118,7 +118,7 @@ inline const FVector2D& FBlendRangesData::GetFirstInBlendRange(const FContiguous
 	
 inline const FVector2D& FBlendRangesData::GetBeforeCurrent(const FContiguousKeys& InKeys, int32 InIndex) const
 {
-	return InIndex > 0 ? AllKeyPositions[InKeys.Indices[InIndex]] : GetBeforeBlendRange(InKeys);
+	return InIndex > 0 ? AllKeyPositions[InKeys.Indices[InIndex - 1]] : GetBeforeBlendRange(InKeys);
 }
 
 inline const FVector2D& FBlendRangesData::GetCurrent(const FContiguousKeys& InKeys, int32 InIndex) const
