Tweening Utils: Add hotkeys to change slider position. By default:
- Home -> 100% left
- End -> 100% right
- Shift+Home -> 25% left
- Shift+End -> 25% right

#Synced-CL 51881412
#Synced-CL 51881412

diff --git a/Engine/Plugins/Animation/ControlRig/Source/ControlRigEditor/Private/EditMode/Tween/TweenLogic.cpp b/Engine/Plugins/Animation/ControlRig/Source/ControlRigEditor/Private/EditMode/Tween/TweenLogic.cpp
--- a/Engine/Plugins/Animation/ControlRig/Source/ControlRigEditor/Private/EditMode/Tween/TweenLogic.cpp
+++ b/Engine/Plugins/Animation/ControlRig/Source/ControlRigEditor/Private/EditMode/Tween/TweenLogic.cpp
@@ -22,7 +22,7 @@ FTweenLogicWidgets FTweenLogic::MakeWidget() const
 	TAttribute<TOptional<float>> OverrideSliderPosition =
 		TAttribute<TOptional<float>>::CreateLambda([this]()
 		{
-			return Controllers.MouseSlidingController.GetCurrentSliderPosition();
+			return Controllers.GetCurrentSliderPosition();
 		});
 	
 	FSlimHorizontalToolBarBuilder ToolBarBuilder(CommandList, FMultiBoxCustomization::None);
diff --git a/Engine/Plugins/Animation/TweeningUtils/Source/TweeningUtilsEditor/Private/Misc/CommandUtils.cpp b/Engine/Plugins/Animation/TweeningUtils/Source/TweeningUtilsEditor/Private/Misc/CommandUtils.cpp
new file mode 100644
--- /dev/null
+++ b/Engine/Plugins/Animation/TweeningUtils/Source/TweeningUtilsEditor/Private/Misc/CommandUtils.cpp
@@ -0,0 +1,91 @@
+﻿// Copyright Epic Games, Inc. All Rights Reserved.
+
+#include "CommandUtils.h"
+
+#include "Framework/Application/IInputProcessor.h"
+#include "Framework/Application/SlateApplication.h"
+#include "Framework/Commands/UICommandInfo.h"
+#include "Input/Events.h"
+
+namespace UE::TweeningUtilsEditor
+{
+namespace Private
+{
+/** Listens for a command's chord going up and executes a delegate. */
+class FCommandChordUpInputProcessor : public IInputProcessor, public TSharedFromThis<FCommandChordUpInputProcessor>
+{
+	const TSharedRef<FUICommandInfo> Command;
+	const FExecuteAction Action;
+
+	bool IsDragCommandKeyUp(const FKeyEvent& InKeyEvent) const
+	{
+		bool bIsMovingSlider = false;
+		
+		for (uint32 i = 0; i < static_cast<uint8>(EMultipleKeyBindingIndex::NumChords); ++i)
+		{
+			EMultipleKeyBindingIndex ChordIndex = static_cast<EMultipleKeyBindingIndex>(i);
+			const FInputChord& Chord = *Command->GetActiveChord(ChordIndex);
+			bIsMovingSlider |= Chord.IsValidChord() && InKeyEvent.GetKey() == Chord.Key;
+		}
+		
+		return bIsMovingSlider;
+	}
+	
+public:
+
+	FCommandChordUpInputProcessor(const TSharedRef<FUICommandInfo>& InCommand, FExecuteAction InAction) 
+		: Command(InCommand)
+		, Action(MoveTemp(InAction))
+	{}
+	
+	//~ Begin IInputProcessor Interface
+	virtual bool HandleKeyUpEvent(FSlateApplication&, const FKeyEvent& InKeyEvent) override
+	{
+		const bool bStop = IsDragCommandKeyUp(InKeyEvent);
+		
+		if (bStop)
+		{
+			FSlateApplication::Get().UnregisterInputPreProcessor(SharedThis(this));
+			Action.Execute();
+		}
+		
+		return bStop;
+	}
+
+	virtual void Tick(const float DeltaTime, FSlateApplication& SlateApp, TSharedRef<ICursor> Cursor) override {}
+	//~ Begin IInputProcessor Interface
+};
+
+class FAutoRegisterInputProcessor : public FCommandChordUpListener
+{
+	const TSharedRef<IInputProcessor> Processor;
+public:
+	
+	explicit FAutoRegisterInputProcessor(const TSharedRef<IInputProcessor>& InProcessor)
+		: Processor(InProcessor)
+	{
+		FSlateApplication::Get().RegisterInputPreProcessor(Processor);
+	}
+	
+	~FAutoRegisterInputProcessor()
+	{
+		// Calling this multiple is fine... it is fine if the processor deregistered itself.
+		FSlateApplication::Get().UnregisterInputPreProcessor(Processor);
+	}
+};
+}
+
+TUniquePtr<FCommandChordUpListener> ListenForCommandChordUp(const TSharedRef<FUICommandInfo>& InCommand, FExecuteAction InAction)
+{
+	using namespace Private;
+	return MakeUnique<FAutoRegisterInputProcessor>(
+		MakeShared<FCommandChordUpInputProcessor>(InCommand, InAction)
+		);
+}
+
+TUniquePtr<FCommandChordUpListener> ListenForCommandChordUp(const TSharedPtr<FUICommandInfo>& InCommand, FExecuteAction InAction)
+{
+	return ListenForCommandChordUp(InCommand.ToSharedRef(), InAction);
+}
+}
+
diff --git a/Engine/Plugins/Animation/TweeningUtils/Source/TweeningUtilsEditor/Private/Misc/CommandUtils.h b/Engine/Plugins/Animation/TweeningUtils/Source/TweeningUtilsEditor/Private/Misc/CommandUtils.h
new file mode 100644
--- /dev/null
+++ b/Engine/Plugins/Animation/TweeningUtils/Source/TweeningUtilsEditor/Private/Misc/CommandUtils.h
@@ -0,0 +1,21 @@
+﻿// Copyright Epic Games, Inc. All Rights Reserved.
+
+#pragma once
+
+#include "Framework/Commands/UIAction.h"
+#include "Templates/SharedPointerFwd.h"
+#include "Templates/UniquePtr.h"
+
+class FUICommandInfo;
+
+namespace UE::TweeningUtilsEditor
+{
+class FCommandChordUpListener { public: virtual ~FCommandChordUpListener() = default; };
+
+/**
+ * Listens for the command chord going up as there is no built-in system for this.
+ * @return Listener object. When it goes out of scope, the subscription is cleaned up.
+ */
+TUniquePtr<FCommandChordUpListener> ListenForCommandChordUp(const TSharedRef<FUICommandInfo>& InCommand, FExecuteAction InAction);
+TUniquePtr<FCommandChordUpListener> ListenForCommandChordUp(const TSharedPtr<FUICommandInfo>& InCommand, FExecuteAction InAction);
+}
diff --git a/Engine/Plugins/Animation/TweeningUtils/Source/TweeningUtilsEditor/Private/TweeningUtilsCommands.cpp b/Engine/Plugins/Animation/TweeningUtils/Source/TweeningUtilsEditor/Private/TweeningUtilsCommands.cpp
--- a/Engine/Plugins/Animation/TweeningUtils/Source/TweeningUtilsEditor/Private/TweeningUtilsCommands.cpp
+++ b/Engine/Plugins/Animation/TweeningUtils/Source/TweeningUtilsEditor/Private/TweeningUtilsCommands.cpp
@@ -34,6 +34,13 @@ void FTweeningUtilsCommands::RegisterCommands()
 	
 	UI_COMMAND(DragAnimSliderTool, "Drag Anim Slider Tool", "Drag the anim slider", EUserInterfaceActionType::Button, FInputChord(EKeys::U));
 	UI_COMMAND(ChangeAnimSliderTool, "Change Anim Slider Tool", "Go to the next blend function", EUserInterfaceActionType::Button, FInputChord(EKeys::U, EModifierKey::Shift));
+	
+	UI_COMMAND(SetSliderLeft, "Set Slider left (100%)", "Drag the anim slider 100% to the left", EUserInterfaceActionType::Button, FInputChord(EKeys::Home));
+	UI_COMMAND(SetSliderRight, "Set Slider right (100%)", "Drag the anim slider 100% to the right", EUserInterfaceActionType::Button, FInputChord(EKeys::End));
+	UI_COMMAND(SetSliderHalfLeft, "Set Slider half left (50%)", "Drag the anim slider 50% to the left", EUserInterfaceActionType::Button, FInputChord());
+	UI_COMMAND(SetSliderHalfRight, "Set Slider half right (50%)", "Drag the anim slider 50% to the right", EUserInterfaceActionType::Button, FInputChord());
+	UI_COMMAND(SetSliderQuarterLeft, "Set Slider quarter left (25%)", "Drag the anim slider 25% to the left", EUserInterfaceActionType::Button, FInputChord(EKeys::Home, EModifierKey::Shift));
+	UI_COMMAND(SetSliderQuarterRight, "Set Slider quarter right (25%)", "Drag the anim slider 25% to the right", EUserInterfaceActionType::Button, FInputChord(EKeys::End, EModifierKey::Shift));
 }
 }
 
diff --git a/Engine/Plugins/Animation/TweeningUtils/Source/TweeningUtilsEditor/Private/Widgets/MVC/MoveSliderByHotkeyController.cpp b/Engine/Plugins/Animation/TweeningUtils/Source/TweeningUtilsEditor/Private/Widgets/MVC/MoveSliderByHotkeyController.cpp
new file mode 100644
--- /dev/null
+++ b/Engine/Plugins/Animation/TweeningUtils/Source/TweeningUtilsEditor/Private/Widgets/MVC/MoveSliderByHotkeyController.cpp
@@ -0,0 +1,86 @@
+﻿// Copyright Epic Games, Inc. All Rights Reserved.
+
+#include "Widgets/MVC/MoveSliderByHotkeyController.h"
+
+#include "Misc/CommandUtils.h"
+#include "Templates/UniquePtr.h"
+#include "TweeningUtilsCommands.h"
+#include "Math/Models/TweenModel.h"
+
+namespace UE::TweeningUtilsEditor
+{
+struct FMoveSliderByHotkeyController::FActiveCommandOverride
+{
+	const float SliderValueOverride;
+	
+	/** Cleans up the listener for the command chord going up. */
+	const TUniquePtr<FCommandChordUpListener> ChordUpListener;
+
+	explicit FActiveCommandOverride(float InSliderValueOverride, TUniquePtr<FCommandChordUpListener> ChordUpListener)
+		: SliderValueOverride(InSliderValueOverride)
+		, ChordUpListener(MoveTemp(ChordUpListener))
+	{}
+};
+
+FMoveSliderByHotkeyController::FMoveSliderByHotkeyController(
+	TAttribute<FTweenModel*> InTweenModelAttr, const TSharedRef<FUICommandList>& InCommandList
+	)
+	: TweenModelAttr(MoveTemp(InTweenModelAttr))
+	, WeakCommandList(InCommandList)
+{
+	FTweeningUtilsCommands& Commands = FTweeningUtilsCommands::Get();
+	const auto MapAction = [this, &InCommandList](const TSharedPtr<FUICommandInfo> Command, float SliderPosition)
+	{
+		InCommandList->MapAction(
+			Command, 
+			FExecuteAction::CreateRaw(this, &FMoveSliderByHotkeyController::HandleMoveSlider, SliderPosition, Command.ToWeakPtr())
+			);
+	};
+	
+	MapAction(Commands.SetSliderLeft, -1.f);
+	MapAction(Commands.SetSliderRight, 1.f);
+	MapAction(Commands.SetSliderHalfLeft, -0.5f);
+	MapAction(Commands.SetSliderHalfRight, 0.5f);
+	MapAction(Commands.SetSliderQuarterLeft, -0.25f);
+	MapAction(Commands.SetSliderQuarterRight, 0.25f);
+}
+
+FMoveSliderByHotkeyController::~FMoveSliderByHotkeyController()
+{
+	if (const TSharedPtr<FUICommandList> CommandList = WeakCommandList.Pin(); CommandList && FTweeningUtilsCommands::IsRegistered())
+	{
+		FTweeningUtilsCommands& Commands = FTweeningUtilsCommands::Get();
+		CommandList->UnmapAction(Commands.SetSliderLeft);
+		CommandList->UnmapAction(Commands.SetSliderRight);
+		CommandList->UnmapAction(Commands.SetSliderHalfLeft);
+		CommandList->UnmapAction(Commands.SetSliderHalfRight);
+		CommandList->UnmapAction(Commands.SetSliderQuarterLeft);
+		CommandList->UnmapAction(Commands.SetSliderQuarterRight);
+	}
+}
+
+TOptional<float> FMoveSliderByHotkeyController::GetCurrentSliderPosition() const
+{
+	return ActiveCommand ? ActiveCommand->SliderValueOverride : TOptional<float>();
+}
+
+void FMoveSliderByHotkeyController::HandleMoveSlider(float InBlendValue, TWeakPtr<FUICommandInfo> InCommandForChordUp)
+{
+	if (FTweenModel* TweenModel = TweenModelAttr.Get())
+	{
+		FExecuteAction HandleCommandChordUpAction = FExecuteAction::CreateRaw(this, &FMoveSliderByHotkeyController::HandleCommandChordUp);
+		ActiveCommand = MakePimpl<FActiveCommandOverride>(
+			InBlendValue, 
+			ListenForCommandChordUp(InCommandForChordUp.Pin(), MoveTemp(HandleCommandChordUpAction))
+			);
+		
+		TweenModel->BlendOneOff(InBlendValue);
+	}
+}
+
+void FMoveSliderByHotkeyController::HandleCommandChordUp()
+{
+	// The point of this is to make the slider appear at the override position for as long as the command keys are pressed down... as visual feedback.
+	ActiveCommand.Reset();
+}
+}
diff --git a/Engine/Plugins/Animation/TweeningUtils/Source/TweeningUtilsEditor/Private/Widgets/MVC/TweenControllers.cpp b/Engine/Plugins/Animation/TweeningUtils/Source/TweeningUtilsEditor/Private/Widgets/MVC/TweenControllers.cpp
--- a/Engine/Plugins/Animation/TweeningUtils/Source/TweeningUtilsEditor/Private/Widgets/MVC/TweenControllers.cpp
+++ b/Engine/Plugins/Animation/TweeningUtils/Source/TweeningUtilsEditor/Private/Widgets/MVC/TweenControllers.cpp
@@ -53,8 +53,24 @@ FTweenControllers::FTweenControllers(
 		InCommandList
 		)
 	, MouseSlidingController(ToolbarController.MakeSelectedTweenModelAttr(), InCommandList)
+	, MoveSliderByHotkeyController(ToolbarController.MakeSelectedTweenModelAttr(), InCommandList)
 {
 	check(InTweenModels->NumModels() > 0);
 	Private::SetupUserPreferredFunctionConfig(ToolbarController, InTweenModels, UserPreferredFunctionContext);
 }
+
+TOptional<float> FTweenControllers::GetCurrentSliderPosition() const
+{
+	if (const TOptional<float> MouseOverride = MouseSlidingController.GetCurrentSliderPosition())
+	{
+		return MouseOverride;
+	}
+	
+	if (const TOptional<float> CommandOverride = MoveSliderByHotkeyController.GetCurrentSliderPosition())
+	{
+		return CommandOverride;
+	}
+	
+	return {};
+}
 }
diff --git a/Engine/Plugins/Animation/TweeningUtils/Source/TweeningUtilsEditor/Public/TweeningUtilsCommands.h b/Engine/Plugins/Animation/TweeningUtils/Source/TweeningUtilsEditor/Public/TweeningUtilsCommands.h
--- a/Engine/Plugins/Animation/TweeningUtils/Source/TweeningUtilsEditor/Public/TweeningUtilsCommands.h
+++ b/Engine/Plugins/Animation/TweeningUtils/Source/TweeningUtilsEditor/Public/TweeningUtilsCommands.h
@@ -19,8 +19,23 @@ public:
 	TSharedPtr<FUICommandInfo> SetTweenTimeOffset;
 	TSharedPtr<FUICommandInfo> ToggleOvershootMode;
 
+	/** Triggers detection of mouse movement to drive the slider indirectly. */
 	TSharedPtr<FUICommandInfo> DragAnimSliderTool;
+	/** Cycles to the next tween functions */
 	TSharedPtr<FUICommandInfo> ChangeAnimSliderTool;
+	
+	/** Positions the slider 100% to the left as if you had dragged it there. */
+	TSharedPtr<FUICommandInfo> SetSliderLeft;
+	/** Positions the slider 100% to the right as if you had dragged it there. */
+	TSharedPtr<FUICommandInfo> SetSliderRight;
+	/** Positions the slider 50% to the left as if you had dragged it there. */
+	TSharedPtr<FUICommandInfo> SetSliderHalfLeft;
+	/** Positions the slider 50% to the right as if you had dragged it there. */
+	TSharedPtr<FUICommandInfo> SetSliderHalfRight;
+	/** Positions the slider 25% to the left as if you had dragged it there. */
+	TSharedPtr<FUICommandInfo> SetSliderQuarterLeft;
+	/** Positions the slider 25% to the right as if you had dragged it there. */
+	TSharedPtr<FUICommandInfo> SetSliderQuarterRight;
 
 	FTweeningUtilsCommands();
 
diff --git a/Engine/Plugins/Animation/TweeningUtils/Source/TweeningUtilsEditor/Public/Widgets/MVC/MoveSliderByHotkeyController.h b/Engine/Plugins/Animation/TweeningUtils/Source/TweeningUtilsEditor/Public/Widgets/MVC/MoveSliderByHotkeyController.h
new file mode 100644
--- /dev/null
+++ b/Engine/Plugins/Animation/TweeningUtils/Source/TweeningUtilsEditor/Public/Widgets/MVC/MoveSliderByHotkeyController.h
@@ -0,0 +1,56 @@
+﻿// Copyright Epic Games, Inc. All Rights Reserved.
+
+#pragma once
+
+#include "Misc/Attribute.h"
+#include "Templates/PimplPtr.h"
+#include "Templates/SharedPointer.h"
+#include "Templates/UnrealTemplate.h"
+
+#define UE_API TWEENINGUTILSEDITOR_API
+
+class FUICommandInfo;
+class FUICommandList;
+
+namespace UE::TweeningUtilsEditor
+{
+class FTweenModel;
+class FTweenMouseSlidingController;
+
+/**
+ * Allows the user to move the slider
+ * - 100% to the left (FTweeningUtilsCommands::SetSliderLeft)
+ * - 100% to the right (FTweeningUtilsCommands::SetSliderRight)
+ * while indirectly moving the slider (@see FTweeningUtilsCommands::DragAnimSliderTool).
+ */
+class FMoveSliderByHotkeyController : public FNoncopyable
+{
+public:
+	
+	UE_API explicit FMoveSliderByHotkeyController(TAttribute<FTweenModel*> InTweenModelAttr, const TSharedRef<FUICommandList>& InCommandList);
+	UE_API ~FMoveSliderByHotkeyController();
+	
+	/** @return The current slider position that is being driven by the active command. Range [-1,1]. */
+	UE_API TOptional<float> GetCurrentSliderPosition() const;
+
+private:
+	
+	/** Does the actual blending. */
+	const TAttribute<FTweenModel*> TweenModelAttr;
+	
+	/** The command list that the commands are bound to. */
+	const TWeakPtr<FUICommandList> WeakCommandList;
+	
+	struct FActiveCommandOverride;
+	/** Set while one of the commands is pressed down. Becomes unset when the command chord goes up. */
+	TPimplPtr<FActiveCommandOverride> ActiveCommand;
+	
+	/** Handles the command to move the slider to InBlendValue and keeps the slider position there until the key is released. */
+	void HandleMoveSlider(float InBlendValue, TWeakPtr<FUICommandInfo> InCommandForChordUp);
+	
+	/** Handles on of the command chord being released. */
+	void HandleCommandChordUp();
+};
+}
+
+#undef UE_API
diff --git a/Engine/Plugins/Animation/TweeningUtils/Source/TweeningUtilsEditor/Public/Widgets/MVC/TweenControllers.h b/Engine/Plugins/Animation/TweeningUtils/Source/TweeningUtilsEditor/Public/Widgets/MVC/TweenControllers.h
--- a/Engine/Plugins/Animation/TweeningUtils/Source/TweeningUtilsEditor/Public/Widgets/MVC/TweenControllers.h
+++ b/Engine/Plugins/Animation/TweeningUtils/Source/TweeningUtilsEditor/Public/Widgets/MVC/TweenControllers.h
@@ -3,6 +3,7 @@
 #pragma once
 
 #include "CycleFunctionController.h"
+#include "MoveSliderByHotkeyController.h"
 #include "Templates/SharedPointerFwd.h"
 #include "TweenMouseSlidingController.h"
 #include "TweenToolbarController.h"
@@ -24,6 +25,9 @@ struct FTweenControllers
 
 	/** Allows indirect movement of the slider by using u + Move Mouse */
 	FTweenMouseSlidingController MouseSlidingController;
+	
+	/** Allows moving the slider via hotkeys. */
+	FMoveSliderByHotkeyController MoveSliderByHotkeyController;
 
 	/**
 	 * @param InCommandList The command list to bind commands to
@@ -35,6 +39,9 @@ struct FTweenControllers
 		const TSharedRef<ITweenModelContainer>& InTweenModels,
 		FName UserPreferredFunctionContext = NAME_None
 		);
+	
+	/** @return The current slider override position that may be driven by one of the subsytems, like MouseSlidingController. Range [-1,1]. */
+	UE_API TOptional<float> GetCurrentSliderPosition() const;
 };
 }
 
diff --git a/Engine/Plugins/Editor/CurveEditorTools/Source/CurveEditorTools/Private/Extensions/Tweening/TweenEditorExtension.cpp b/Engine/Plugins/Editor/CurveEditorTools/Source/CurveEditorTools/Private/Extensions/Tweening/TweenEditorExtension.cpp
--- a/Engine/Plugins/Editor/CurveEditorTools/Source/CurveEditorTools/Private/Extensions/Tweening/TweenEditorExtension.cpp
+++ b/Engine/Plugins/Editor/CurveEditorTools/Source/CurveEditorTools/Private/Extensions/Tweening/TweenEditorExtension.cpp
@@ -37,7 +37,7 @@ TSharedPtr<FExtender> FTweenEditorExtension::MakeToolbarExtender(const TSharedRe
 				TAttribute<TOptional<float>> OverrideSliderPosition =
 					TAttribute<TOptional<float>>::CreateLambda([this]()
 					{
-						return TweenControllers->MouseSlidingController.GetCurrentSliderPosition();
+						return TweenControllers->GetCurrentSliderPosition();
 					});
 				
 				ToolbarBuilder.BeginStyleOverride(TEXT("CurveEditorTweenToolbar")); // Forces min sizes for combo button
