Curve Editor: Add Tweening.Ease.SlopeExponent and Tweening.Ease.SlopeMultiplier CVars to better control the ease function.

This is related to UE-373159. These CVars will be adjusted with better values later once animators have given feedback what value is good.

#Synced-CL 52378315

diff --git a/Engine/Plugins/Animation/TweeningUtils/Source/TweeningUtils/Private/Math/KeyBlendingFunctions.cpp b/Engine/Plugins/Animation/TweeningUtils/Source/TweeningUtils/Private/Math/KeyBlendingFunctions.cpp
--- a/Engine/Plugins/Animation/TweeningUtils/Source/TweeningUtils/Private/Math/KeyBlendingFunctions.cpp
+++ b/Engine/Plugins/Animation/TweeningUtils/Source/TweeningUtils/Private/Math/KeyBlendingFunctions.cpp
@@ -2,6 +2,7 @@
 
 #include "Math/KeyBlendingFunctions.h"
 
+#include "HAL/IConsoleManager.h"
 #include "Math/Vector2D.h"
 
 namespace UE::TweeningUtils
@@ -69,6 +70,21 @@ double Blend_Relative(double InBlendValue,
 
 namespace Private::BlendToEase
 { 
+
+static TAutoConsoleVariable<float> CVarEaseSlopeMultiplier(
+	  TEXT("Tweening.Ease.SlopeMultiplier"),
+	  5.0f,
+	  TEXT("Controls the maximum slope of the Ease blend curve (at 100% slider). Higher = sharper peak curve.")
+	  );
+
+static TAutoConsoleVariable<float> CVarEaseSlopeExponent(
+	TEXT("Tweening.Ease.SlopeExponent"),
+	1.0f,
+	TEXT("Controls how aggressively the Ease slope grows with slider position. "
+		 "1.0 = linear, 2.0 = smoother behavior (gentler at low slider, sharper only at high slider).")
+		 );
+
+
 static float SCurve(float X, float Slope, float Width, float Height, float XShift, float YShift)
 {
 	if (X > (XShift + Width))
@@ -96,7 +112,6 @@ double Blend_Ease(
 	}
 	const double Source = InCurrent.Y;
 	const double FullTimeDiff = InAfterBlendRange.X - InBeforeBlendRange.X;
-	const double AbsValue = FMath::Abs(InBlendValue);
 	const double X = InCurrent.X - InBeforeBlendRange.X;
 	const double Ratio = X / FullTimeDiff;
 	double Shift, Delta, Base;
@@ -112,7 +127,12 @@ double Blend_Ease(
 		Delta = Source - InBeforeBlendRange.Y;
 		Base = InBeforeBlendRange.Y;
 	}
-	const double Slope = 5.0 * AbsValue;
+	
+	const float SlopeMultiplier = Private::BlendToEase::CVarEaseSlopeMultiplier.GetValueOnAnyThread();
+	const float SlopeExponent   = Private::BlendToEase::CVarEaseSlopeExponent.GetValueOnAnyThread();
+	const double AbsValue = FMath::Abs(InBlendValue);
+	const double Slope = SlopeMultiplier * FMath::Pow(AbsValue, SlopeExponent);
+	
 	const double EaseY = Private::BlendToEase::SCurve(Ratio, Slope, 2.0, 2.0, Shift, Shift);
 	const double NewValue = Base + (Delta * EaseY);
 
