Tweening Utils: Improve parameters for easing tween function based on animator feedback.

#Synced-CL 52617198

diff --git a/Engine/Plugins/Animation/TweeningUtils/Source/TweeningUtils/Private/Math/KeyBlendingFunctions.cpp b/Engine/Plugins/Animation/TweeningUtils/Source/TweeningUtils/Private/Math/KeyBlendingFunctions.cpp
--- a/Engine/Plugins/Animation/TweeningUtils/Source/TweeningUtils/Private/Math/KeyBlendingFunctions.cpp
+++ b/Engine/Plugins/Animation/TweeningUtils/Source/TweeningUtils/Private/Math/KeyBlendingFunctions.cpp
@@ -73,13 +73,13 @@ namespace Private::BlendToEase
 
 static TAutoConsoleVariable<float> CVarEaseSlopeMultiplier(
 	  TEXT("Tweening.Ease.SlopeMultiplier"),
-	  5.0f,
+	  6.0f,
 	  TEXT("Controls the maximum slope of the Ease blend curve (at 100% slider). Higher = sharper peak curve.")
 	  );
 
 static TAutoConsoleVariable<float> CVarEaseSlopeExponent(
 	TEXT("Tweening.Ease.SlopeExponent"),
-	1.0f,
+	2.0f,
 	TEXT("Controls how aggressively the Ease slope grows with slider position. "
 		 "1.0 = linear, 2.0 = smoother behavior (gentler at low slider, sharper only at high slider).")
 		 );
